CameraCalibration: Move LensComponent into its own plugin
#rnx

diff --git a/Engine/Plugins/Compositing/Composure/Composure.uplugin b/Engine/Plugins/Compositing/Composure/Composure.uplugin
--- a/Engine/Plugins/Compositing/Composure/Composure.uplugin
+++ b/Engine/Plugins/Compositing/Composure/Composure.uplugin
@@ -52,6 +52,10 @@
 		{
 			"Name": "ActorLayerUtilities",
 			"Enabled": true
+		},
+		{
+			"Name": "LensComponent",
+			"Enabled": true
 		}
 	]
 }
\ No newline at end of file
diff --git a/Engine/Plugins/Compositing/Composure/Source/Composure/Composure.Build.cs b/Engine/Plugins/Compositing/Composure/Source/Composure/Composure.Build.cs
--- a/Engine/Plugins/Compositing/Composure/Source/Composure/Composure.Build.cs
+++ b/Engine/Plugins/Compositing/Composure/Source/Composure/Composure.Build.cs
@@ -22,6 +22,7 @@ namespace UnrealBuildTool.Rules
 					"Core",
                     "CoreUObject",
                     "Engine",
+					"LensComponent",
 					"MediaIOCore",
 					"MovieScene",
 					"MovieSceneTracks",
diff --git a/Engine/Plugins/VirtualProduction/CameraCalibration/CameraCalibration.uplugin b/Engine/Plugins/VirtualProduction/CameraCalibration/CameraCalibration.uplugin
--- a/Engine/Plugins/VirtualProduction/CameraCalibration/CameraCalibration.uplugin
+++ b/Engine/Plugins/VirtualProduction/CameraCalibration/CameraCalibration.uplugin
@@ -29,6 +29,10 @@
 			"Name": "Composure",
 			"Enabled": true
 		},
+		{
+			"Name": "LensComponent",
+			"Enabled": true
+		},
 		{
 			"Name": "MediaFrameworkUtilities",
 			"Enabled": true
diff --git a/Engine/Plugins/VirtualProduction/CameraCalibration/Source/CameraCalibrationEditor/CameraCalibrationEditor.Build.cs b/Engine/Plugins/VirtualProduction/CameraCalibration/Source/CameraCalibrationEditor/CameraCalibrationEditor.Build.cs
--- a/Engine/Plugins/VirtualProduction/CameraCalibration/Source/CameraCalibrationEditor/CameraCalibrationEditor.Build.cs
+++ b/Engine/Plugins/VirtualProduction/CameraCalibration/Source/CameraCalibrationEditor/CameraCalibrationEditor.Build.cs
@@ -31,6 +31,7 @@ namespace UnrealBuildTool.Rules
 					"InputCore",
 					"Json",
 					"JsonUtilities",
+					"LensComponent",
 					"MediaAssets",
 					"MediaFrameworkUtilities",
 					"OpenCV",
diff --git a/Engine/Plugins/VirtualProduction/CameraCalibrationCore/CameraCalibrationCore.uplugin b/Engine/Plugins/VirtualProduction/CameraCalibrationCore/CameraCalibrationCore.uplugin
--- a/Engine/Plugins/VirtualProduction/CameraCalibrationCore/CameraCalibrationCore.uplugin
+++ b/Engine/Plugins/VirtualProduction/CameraCalibrationCore/CameraCalibrationCore.uplugin
@@ -27,28 +27,12 @@
 			"Type": "Editor",
 			"ProgramAllowlist": [ "LiveLinkHub" ],
 			"LoadingPhase": "Default"
-		},
-		{
-			"Name": "CameraCalibrationCoreMovieScene",
-			"Type": "Runtime",
-			"ProgramAllowlist": [ "LiveLinkHub" ],
-			"LoadingPhase": "Default"
-		},
-		{
-			"Name": "CameraCalibrationCoreSequencer",
-			"Type": "UncookedOnly",
-			"ProgramAllowlist": [ "LiveLinkHub" ],
-			"LoadingPhase": "Default"
 		}
 	],
 	"Plugins": [
 		{
 			"Name": "ProceduralMeshComponent",
 			"Enabled": true
-		},
-		{
-			"Name": "LiveLink",
-			"Enabled": true
 		}
 	]
-}
\ No newline at end of file
+}
diff --git a/Engine/Plugins/VirtualProduction/CameraCalibrationCore/Source/CameraCalibrationCore/CameraCalibrationCore.Build.cs b/Engine/Plugins/VirtualProduction/CameraCalibrationCore/Source/CameraCalibrationCore/CameraCalibrationCore.Build.cs
--- a/Engine/Plugins/VirtualProduction/CameraCalibrationCore/Source/CameraCalibrationCore/CameraCalibrationCore.Build.cs
+++ b/Engine/Plugins/VirtualProduction/CameraCalibrationCore/Source/CameraCalibrationCore/CameraCalibrationCore.Build.cs
@@ -14,7 +14,6 @@ namespace UnrealBuildTool.Rules
 					"CoreUObject",
 					"DeveloperSettings",
 					"Engine",
-					"LiveLinkComponents",
 					"RenderCore",
 					"RHI",
 					"Slate",
diff --git a/Engine/Plugins/VirtualProduction/CameraCalibrationCore/Source/CameraCalibrationCore/Private/CameraCalibrationSubsystem.cpp b/Engine/Plugins/VirtualProduction/CameraCalibrationCore/Source/CameraCalibrationCore/Private/CameraCalibrationSubsystem.cpp
--- a/Engine/Plugins/VirtualProduction/CameraCalibrationCore/Source/CameraCalibrationCore/Private/CameraCalibrationSubsystem.cpp
+++ b/Engine/Plugins/VirtualProduction/CameraCalibrationCore/Source/CameraCalibrationCore/Private/CameraCalibrationSubsystem.cpp
@@ -7,7 +7,6 @@
 #include "CineCameraComponent.h"
 #include "Engine/TimecodeProvider.h"
 #include "GameFramework/Actor.h"
-#include "LensComponent.h"
 #include "Misc/CoreDelegates.h"
 #include "UObject/UObjectIterator.h"
 
@@ -46,20 +45,6 @@ TArray<ULensDistortionModelHandlerBase*> UCameraCalibrationSubsystem::GetDistort
 {
 	// This function has been deprecated. The implementation has been changed to provide some backwards compatibility, but code should be updated to not call this function.
 	TArray<ULensDistortionModelHandlerBase*> Handlers;
-
-	if (Component)
-	{
-		AActor* OwningActor = Component->GetOwner();
-
-		TInlineComponentArray<ULensComponent*> LensComponents;
-		OwningActor->GetComponents(LensComponents);
-
-		for (ULensComponent* LensComponent : LensComponents)
-		{
-			Handlers.Add(LensComponent->GetLensDistortionHandler());
-		}
-	}
-
 	return Handlers;
 }
 
@@ -68,22 +53,6 @@ ULensDistortionModelHandlerBase* UCameraCalibrationSubsystem::FindDistortionMode
 {
 	// This function has been deprecated. The implementation has been changed to provide some backwards compatibility, but code should be updated to not call this function.
 	TArray<ULensDistortionModelHandlerBase*> Handlers;
-
-	if (!DistortionHandlerPicker.TargetCameraComponent)
-	{
-		return nullptr;
-	}
-
-	AActor* OwningActor = DistortionHandlerPicker.TargetCameraComponent->GetOwner();
-
-	TInlineComponentArray<ULensComponent*> LensComponents;
-	OwningActor->GetComponents(LensComponents);
-
-	if (LensComponents.Num() > 0)
-	{
-		return LensComponents[0]->GetLensDistortionHandler();
-	}
-
 	return nullptr;
 }
 
diff --git a/Engine/Plugins/VirtualProduction/CameraCalibrationCore/Source/CameraCalibrationCore/Private/LensComponent.cpp b/Engine/Plugins/VirtualProduction/CameraCalibrationCore/Source/CameraCalibrationCore/Private/LensComponent.cpp
deleted file mode 100644
diff --git a/Engine/Plugins/VirtualProduction/CameraCalibrationCore/Source/CameraCalibrationCore/Public/LensComponent.h b/Engine/Plugins/VirtualProduction/CameraCalibrationCore/Source/CameraCalibrationCore/Public/LensComponent.h
deleted file mode 100644
diff --git a/Engine/Plugins/VirtualProduction/CameraCalibrationCore/Source/CameraCalibrationCoreEditor/Private/CameraCalibrationCoreEditorModule.cpp b/Engine/Plugins/VirtualProduction/CameraCalibrationCore/Source/CameraCalibrationCoreEditor/Private/CameraCalibrationCoreEditorModule.cpp
--- a/Engine/Plugins/VirtualProduction/CameraCalibrationCore/Source/CameraCalibrationCoreEditor/Private/CameraCalibrationCoreEditorModule.cpp
+++ b/Engine/Plugins/VirtualProduction/CameraCalibrationCore/Source/CameraCalibrationCoreEditor/Private/CameraCalibrationCoreEditorModule.cpp
@@ -12,13 +12,10 @@
 #include "CameraCalibrationCoreEditorStyle.h"
 #include "Editor.h"
 #include "IPlacementModeModule.h"
-#include "LensComponent.h"
-#include "LensComponentDetailCustomization.h"
-#include "LensFile.h"
+#include "LevelEditorOutlinerSettings.h"
 #include "Misc/CoreDelegates.h"
 #include "Modules/ModuleManager.h"
 #include "PropertyEditorModule.h"
-#include "LevelEditorOutlinerSettings.h"
 
 LLM_DEFINE_TAG(CameraCalibrationCore_CameraCalibrationCoreEditor);
 #define LOCTEXT_NAMESPACE "CameraCalibrationCoreEditor"
@@ -39,11 +36,6 @@ void FCameraCalibrationCoreEditorModule::StartupModule()
 		FOnGetDetailCustomizationInstance::CreateStatic(&FCalibrationPointComponentDetails::MakeInstance)
 	);
 
-	PropertyEditorModule.RegisterCustomClassLayout(
-		ULensComponent::StaticClass()->GetFName(),
-		FOnGetDetailCustomizationInstance::CreateStatic(&FLensComponentDetailCustomization::MakeInstance)
-	);
-
 	RegisterPlacementModeItems();
 }
 
@@ -56,7 +48,6 @@ void FCameraCalibrationCoreEditorModule::ShutdownModule()
 		FPropertyEditorModule& PropertyModule = FModuleManager::GetModuleChecked<FPropertyEditorModule>("PropertyEditor");
 
 		PropertyModule.UnregisterCustomClassLayout(UCalibrationPointComponent::StaticClass()->GetFName());
-		PropertyModule.UnregisterCustomClassLayout(ULensComponent::StaticClass()->GetFName());
 
 		UnregisterPlacementModeItems();
 	}
diff --git a/Engine/Plugins/VirtualProduction/CameraCalibrationCore/Source/CameraCalibrationCoreEditor/Private/LensComponentDetailCustomization.cpp b/Engine/Plugins/VirtualProduction/CameraCalibrationCore/Source/CameraCalibrationCoreEditor/Private/LensComponentDetailCustomization.cpp
deleted file mode 100644
diff --git a/Engine/Plugins/VirtualProduction/CameraCalibrationCore/Source/CameraCalibrationCoreEditor/Private/LensComponentDetailCustomization.h b/Engine/Plugins/VirtualProduction/CameraCalibrationCore/Source/CameraCalibrationCoreEditor/Private/LensComponentDetailCustomization.h
deleted file mode 100644
diff --git a/Engine/Plugins/VirtualProduction/CameraCalibrationCore/Source/CameraCalibrationCoreEditor/Private/LensDistortionStateDetailCustomization.h b/Engine/Plugins/VirtualProduction/CameraCalibrationCore/Source/CameraCalibrationCoreEditor/Private/LensDistortionStateDetailCustomization.h
deleted file mode 100644
diff --git a/Engine/Plugins/VirtualProduction/CameraCalibrationCore/Source/CameraCalibrationCoreEditor/Public/LensDistortionStateDetailCustomization.h b/Engine/Plugins/VirtualProduction/CameraCalibrationCore/Source/CameraCalibrationCoreEditor/Public/LensDistortionStateDetailCustomization.h
new file mode 100644
--- /dev/null
+++ b/Engine/Plugins/VirtualProduction/CameraCalibrationCore/Source/CameraCalibrationCoreEditor/Public/LensDistortionStateDetailCustomization.h
@@ -0,0 +1,32 @@
+// Copyright Epic Games, Inc. All Rights Reserved.
+
+#pragma once
+
+#include "IPropertyTypeCustomization.h"
+
+#include "Models/LensModel.h"
+#include "Templates/SubclassOf.h"
+
+class CAMERACALIBRATIONCOREEDITOR_API FLensDistortionStateDetailCustomization : public IPropertyTypeCustomization
+{
+public:
+	FLensDistortionStateDetailCustomization(TSubclassOf<ULensModel> InLensModel);
+
+	// Begin IDetailCustomization interface
+	virtual void CustomizeHeader(TSharedRef<IPropertyHandle> PropertyHandle, FDetailWidgetRow& HeaderRow, IPropertyTypeCustomizationUtils& CustomizationUtils) override;
+	virtual void CustomizeChildren(TSharedRef<IPropertyHandle> PropertyHandle, IDetailChildrenBuilder& ChildBuilder, IPropertyTypeCustomizationUtils& CustomizationUtils) override;
+	//~ End IDetailCustomization interface
+
+private:
+	/** Customize the FDistortionInfo struct property */
+	void CustomizeDistortionInfo(TSharedRef<IPropertyHandle> PropertyHandle, IDetailChildrenBuilder& ChildBuilder, IPropertyTypeCustomizationUtils& CustomizationUtils);
+
+	/** Customize the FFocalLengthInfo struct property */
+	void CustomizeFocalLengthInfo(TSharedRef<IPropertyHandle> PropertyHandle, IDetailChildrenBuilder& ChildBuilder, IPropertyTypeCustomizationUtils& CustomizationUtils);
+
+	/** Customize the FImageCenterInfo struct property */
+	void CustomizeImageCenterInfo(TSharedRef<IPropertyHandle> PropertyHandle, IDetailChildrenBuilder& ChildBuilder, IPropertyTypeCustomizationUtils& CustomizationUtils);
+
+	/** The LensModel used to customize the names of the properties in the distortion parameter array */
+	TSubclassOf<ULensModel> LensModel;
+};
diff --git a/Engine/Plugins/VirtualProduction/CameraCalibrationCore/Source/CameraCalibrationCoreMovieScene/CameraCalibrationCoreMovieScene.build.cs b/Engine/Plugins/VirtualProduction/CameraCalibrationCore/Source/CameraCalibrationCoreMovieScene/CameraCalibrationCoreMovieScene.build.cs
deleted file mode 100644
diff --git a/Engine/Plugins/VirtualProduction/CameraCalibrationCore/Source/CameraCalibrationCoreMovieScene/Private/CameraCalibrationCoreMovieScene.cpp b/Engine/Plugins/VirtualProduction/CameraCalibrationCore/Source/CameraCalibrationCoreMovieScene/Private/CameraCalibrationCoreMovieScene.cpp
deleted file mode 100644
diff --git a/Engine/Plugins/VirtualProduction/CameraCalibrationCore/Source/CameraCalibrationCoreMovieScene/Private/MovieSceneLensComponentSection.cpp b/Engine/Plugins/VirtualProduction/CameraCalibrationCore/Source/CameraCalibrationCoreMovieScene/Private/MovieSceneLensComponentSection.cpp
deleted file mode 100644
diff --git a/Engine/Plugins/VirtualProduction/CameraCalibrationCore/Source/CameraCalibrationCoreMovieScene/Private/MovieSceneLensComponentTrack.cpp b/Engine/Plugins/VirtualProduction/CameraCalibrationCore/Source/CameraCalibrationCoreMovieScene/Private/MovieSceneLensComponentTrack.cpp
deleted file mode 100644
diff --git a/Engine/Plugins/VirtualProduction/CameraCalibrationCore/Source/CameraCalibrationCoreMovieScene/Public/MovieSceneLensComponentSection.h b/Engine/Plugins/VirtualProduction/CameraCalibrationCore/Source/CameraCalibrationCoreMovieScene/Public/MovieSceneLensComponentSection.h
deleted file mode 100644
diff --git a/Engine/Plugins/VirtualProduction/CameraCalibrationCore/Source/CameraCalibrationCoreMovieScene/Public/MovieSceneLensComponentTrack.h b/Engine/Plugins/VirtualProduction/CameraCalibrationCore/Source/CameraCalibrationCoreMovieScene/Public/MovieSceneLensComponentTrack.h
deleted file mode 100644
diff --git a/Engine/Plugins/VirtualProduction/CameraCalibrationCore/Source/CameraCalibrationCoreSequencer/CameraCalibrationCoreSequencer.build.cs b/Engine/Plugins/VirtualProduction/CameraCalibrationCore/Source/CameraCalibrationCoreSequencer/CameraCalibrationCoreSequencer.build.cs
deleted file mode 100644
diff --git a/Engine/Plugins/VirtualProduction/CameraCalibrationCore/Source/CameraCalibrationCoreSequencer/Private/CameraCalibrationCoreSequencer.cpp b/Engine/Plugins/VirtualProduction/CameraCalibrationCore/Source/CameraCalibrationCoreSequencer/Private/CameraCalibrationCoreSequencer.cpp
deleted file mode 100644
diff --git a/Engine/Plugins/VirtualProduction/CameraCalibrationCore/Source/CameraCalibrationCoreSequencer/Private/LensComponentTrackEditor.cpp b/Engine/Plugins/VirtualProduction/CameraCalibrationCore/Source/CameraCalibrationCoreSequencer/Private/LensComponentTrackEditor.cpp
deleted file mode 100644
diff --git a/Engine/Plugins/VirtualProduction/CameraCalibrationCore/Source/CameraCalibrationCoreSequencer/Private/LensComponentTrackEditor.h b/Engine/Plugins/VirtualProduction/CameraCalibrationCore/Source/CameraCalibrationCoreSequencer/Private/LensComponentTrackEditor.h
deleted file mode 100644
diff --git a/Engine/Plugins/VirtualProduction/CameraCalibrationCore/Source/CameraCalibrationCoreSequencer/Private/MovieSceneLensComponentTrackRecorder.cpp b/Engine/Plugins/VirtualProduction/CameraCalibrationCore/Source/CameraCalibrationCoreSequencer/Private/MovieSceneLensComponentTrackRecorder.cpp
deleted file mode 100644
diff --git a/Engine/Plugins/VirtualProduction/CameraCalibrationCore/Source/CameraCalibrationCoreSequencer/Private/MovieSceneLensComponentTrackRecorder.h b/Engine/Plugins/VirtualProduction/CameraCalibrationCore/Source/CameraCalibrationCoreSequencer/Private/MovieSceneLensComponentTrackRecorder.h
deleted file mode 100644
diff --git a/Engine/Plugins/VirtualProduction/LensComponent/Config/BaseLensComponent.ini b/Engine/Plugins/VirtualProduction/LensComponent/Config/BaseLensComponent.ini
new file mode 100644
--- /dev/null
+++ b/Engine/Plugins/VirtualProduction/LensComponent/Config/BaseLensComponent.ini
@@ -0,0 +1,12 @@
+[CoreRedirects]
+
+;
+; 5.4 Redirects
+;
+
+; Classes
+
++ClassRedirects=(OldName="/Script/CameraCalibrationCore.LensComponent", NewName="/Script/LensComponent.LensComponent")
++ClassRedirects=(OldName="/Script/CameraCalibrationCoreMovieScene.MovieSceneLensComponentSection", NewName="/Script/LensComponentEditor.MovieSceneLensComponentSection")
++ClassRedirects=(OldName="/Script/CameraCalibrationCoreMovieScene.MovieSceneLensComponentTrack", NewName="/Script/LensComponentEditor.MovieSceneLensComponentTrack")
++ClassRedirects=(OldName="/Script/CameraCalibrationCoreSequencer.MovieSceneLensComponentTrackRecorder", NewName="/Script/LensComponentEditor.MovieSceneLensComponentTrackRecorder")
diff --git a/Engine/Plugins/VirtualProduction/LensComponent/LensComponent.uplugin b/Engine/Plugins/VirtualProduction/LensComponent/LensComponent.uplugin
new file mode 100644
--- /dev/null
+++ b/Engine/Plugins/VirtualProduction/LensComponent/LensComponent.uplugin
@@ -0,0 +1,46 @@
+{
+	"FileVersion": 3,
+	"Version": 1,
+	"VersionName": "1.0",
+	"FriendlyName": "Lens Component",
+	"Description": "Implements the Lens Component for adding distortion to a cinematic camera",
+	"Category": "Virtual Production",
+	"CreatedBy": "Epic Games, Inc.",
+	"CreatedByURL": "https://epicgames.com",
+	"DocsURL": "",
+	"MarketplaceURL": "",
+	"SupportURL": "",
+	"CanContainContent": true,
+	"IsBetaVersion": true,
+	"Hidden": true,
+	"Installed": false,
+	"SupportedPrograms": [ "LiveLinkHub" ],
+	"Modules": [
+		{
+			"Name": "LensComponent",
+			"Type": "Runtime",
+			"ProgramAllowlist": [ "LiveLinkHub" ],
+			"LoadingPhase": "Default"
+		},
+		{
+			"Name": "LensComponentEditor",
+			"Type": "Editor",
+			"ProgramAllowlist": [ "LiveLinkHub" ],
+			"LoadingPhase": "Default"
+		}
+	],
+	"Plugins": [
+		{
+			"Name": "CameraCalibrationCore",
+			"Enabled": true
+		},
+		{
+			"Name": "LiveLink",
+			"Enabled": true
+		},
+		{
+			"Name": "Takes",
+			"Enabled": true
+		}
+	]
+}
\ No newline at end of file
diff --git a/Engine/Plugins/VirtualProduction/LensComponent/Source/LensComponent/LensComponent.build.cs b/Engine/Plugins/VirtualProduction/LensComponent/Source/LensComponent/LensComponent.build.cs
new file mode 100644
--- /dev/null
+++ b/Engine/Plugins/VirtualProduction/LensComponent/Source/LensComponent/LensComponent.build.cs
@@ -0,0 +1,29 @@
+// Copyright Epic Games, Inc. All Rights Reserved.
+
+namespace UnrealBuildTool.Rules
+{
+	public class LensComponent : ModuleRules
+	{
+		public LensComponent(ReadOnlyTargetRules Target) : base(Target)
+		{
+			PublicDependencyModuleNames.AddRange(
+				new string[]
+				{
+					"LiveLinkComponents",
+					"LiveLinkInterface",
+				}
+			);
+
+			PrivateDependencyModuleNames.AddRange(
+				new string[]
+				{
+					"CameraCalibrationCore",
+					"CinematicCamera",
+					"Core",
+					"CoreUObject",
+					"Engine",
+				}
+			);
+		}
+	}
+}
diff --git a/Engine/Plugins/VirtualProduction/LensComponent/Source/LensComponent/Private/LensComponent.cpp b/Engine/Plugins/VirtualProduction/LensComponent/Source/LensComponent/Private/LensComponent.cpp
new file mode 100644
--- /dev/null
+++ b/Engine/Plugins/VirtualProduction/LensComponent/Source/LensComponent/Private/LensComponent.cpp
@@ -0,0 +1,927 @@
+// Copyright Epic Games, Inc. All Rights Reserved.
+
+#include "LensComponent.h"
+
+#include "CameraCalibrationSubsystem.h"
+#include "CineCameraComponent.h"
+#include "Controllers/LiveLinkTransformController.h"
+#include "Engine/Engine.h"
+#include "ILiveLinkComponentModule.h"
+#include "Modules/ModuleManager.h"
+#include "Roles/LiveLinkCameraTypes.h"
+#include "Roles/LiveLinkCameraRole.h"
+#include "Roles/LiveLinkTransformRole.h"
+#include "UObject/UE5MainStreamObjectVersion.h"
+
+DEFINE_LOG_CATEGORY_STATIC(LogLensComponent, Log, All);
+
+PRAGMA_DISABLE_DEPRECATION_WARNINGS
+ULensComponent::ULensComponent()
+{
+	PrimaryComponentTick.bCanEverTick = true;
+	PrimaryComponentTick.bStartWithTickEnabled = true;
+	// This component is designed to tick after the LiveLink component, which uses TG_PrePhysics
+	// We also use a tick prerequisite on LiveLink components, so technically this could also use TG_PrePhysics
+	PrimaryComponentTick.TickGroup = ETickingGroup::TG_DuringPhysics;
+	bTickInEditor = true;
+}
+PRAGMA_ENABLE_DEPRECATION_WARNINGS
+
+void ULensComponent::OnRegister()
+{
+	Super::OnRegister();
+
+	InitDefaultCamera();
+
+	// The component might already have a serialized LensModel, but might not go through PostLoad (in the case that this is spawned by sequencer)
+	// so attempt to create a distortion handler now for the currently set LensModel. 
+	CreateDistortionHandler();
+
+	// Register a callback to let us know when a new LiveLinkComponent is added to this component's parent actor
+	// This gives us the opportunity to track when the subject representation of that LiveLinkComponent changes
+	ILiveLinkComponentsModule& LiveLinkComponentsModule = FModuleManager::GetModuleChecked<ILiveLinkComponentsModule>(TEXT("LiveLinkComponents"));
+	if (!LiveLinkComponentsModule.OnLiveLinkComponentRegistered().IsBoundToObject(this))
+	{
+		LiveLinkComponentsModule.OnLiveLinkComponentRegistered().AddUObject(this, &ULensComponent::OnLiveLinkComponentRegistered);
+	}
+
+	// Look for any LiveLinkComponents that were previously added to this component's parent actor
+	TInlineComponentArray<ULiveLinkComponentController*> LiveLinkComponents;
+	if (AActor* Owner = GetOwner())
+	{
+		Owner->GetComponents(LiveLinkComponents);
+	}
+
+	for (ULiveLinkComponentController* LiveLinkComponent : LiveLinkComponents)
+	{	
+		AddTickPrerequisiteComponent(LiveLinkComponent);
+
+		if (!LiveLinkComponent->OnLiveLinkControllersTicked().IsBoundToObject(this))
+		{
+			LiveLinkComponent->OnLiveLinkControllersTicked().AddUObject(this, &ULensComponent::ProcessLiveLinkData);
+		}
+	}
+}
+
+void ULensComponent::OnUnregister()
+{
+	Super::OnUnregister();
+
+	ILiveLinkComponentsModule& LiveLinkComponentsModule = FModuleManager::GetModuleChecked<ILiveLinkComponentsModule>(TEXT("LiveLinkComponents"));
+	LiveLinkComponentsModule.OnLiveLinkComponentRegistered().RemoveAll(this);
+
+	TInlineComponentArray<ULiveLinkComponentController*> LiveLinkComponents;
+	if (AActor* Owner = GetOwner())
+	{
+		Owner->GetComponents(LiveLinkComponents);
+	}
+
+	for (ULiveLinkComponentController* LiveLinkComponent : LiveLinkComponents)
+	{
+		LiveLinkComponent->OnLiveLinkControllersTicked().RemoveAll(this);
+	}
+}
+
+void ULensComponent::TickComponent(float DeltaTime, ELevelTick TickType, FActorComponentTickFunction* ThisTickFunction)
+{
+	if (UCineCameraComponent* const CineCameraComponent = Cast<UCineCameraComponent>(TargetCameraComponent.GetComponent(GetOwner())))
+	{
+		// If the camera's focal length has changed since the last time this component ticked, we need to update the original focal length we are using
+		if (CineCameraComponent->CurrentFocalLength != LastFocalLength)
+		{
+			// If this is a recorded lens component, then there are no external forces changing the focal length. Original focal length should be preserved.
+			if (EvaluationMode != EFIZEvaluationMode::UseRecordedValues)
+			{
+				OriginalFocalLength = CineCameraComponent->CurrentFocalLength;
+			}
+		}
+
+		// Get the focus and zoom values needed to evaluate the LensFile this Tick
+		UpdateLensFileEvaluationInputs(CineCameraComponent);
+
+		// Attempt to apply nodal offset, which will only succeed if there is a valid LensFile, evaluation inputs, and component to offset
+		if (bApplyNodalOffsetOnTick)
+		{
+			ApplyNodalOffset();
+		}
+
+		ULensFile* LensFile = LensFilePicker.GetLensFile();
+
+		UpdateCameraFilmback(CineCameraComponent);
+
+		EvaluateFocalLength(CineCameraComponent);
+
+		// Evaluate Distortion
+		bWasDistortionEvaluated = false;
+		TObjectPtr<ULensDistortionModelHandlerBase> LensDistortionHandler = LensDistortionHandlerMap.FindRef(LensModel);
+
+		if (LensDistortionHandler)
+		{
+			switch (DistortionStateSource)
+			{
+			case EDistortionSource::LensFile:
+			{
+				if (!LensFile || !EvalInputs.bIsValid)
+				{
+					break;
+				}
+
+				LensFile->EvaluateDistortionData(EvalInputs.Focus, EvalInputs.Zoom, FVector2D(EvalInputs.Filmback.SensorWidth, EvalInputs.Filmback.SensorHeight), LensDistortionHandler);
+
+				DistortionState = LensDistortionHandler->GetCurrentDistortionState();
+
+				// Adjust overscan by the overscan multiplier
+				if (bScaleOverscan)
+				{
+					const float ScaledOverscanFactor = ((LensDistortionHandler->GetOverscanFactor() - 1.0f) * OverscanMultiplier) + 1.0f;
+					LensDistortionHandler->SetOverscanFactor(ScaledOverscanFactor);
+				}
+
+				bWasDistortionEvaluated = true;
+				break;
+			}
+			case EDistortionSource::LiveLinkLensSubject:
+			{
+				// This case is used by the customization to make the distortion state properties read-only. Otherwise, behavior is the same as EDistortionSource::Manual.
+				// falls through
+			}
+			case EDistortionSource::Manual:
+			{
+				LensDistortionHandler->SetDistortionState(DistortionState);
+				LensDistortionHandler->SetCameraFilmback(CineCameraComponent->Filmback);
+
+				//Recompute overscan factor for the distortion state
+				float OverscanFactor = LensDistortionHandler->ComputeOverscanFactor();
+
+				if (bScaleOverscan)
+				{
+					OverscanFactor = ((OverscanFactor - 1.0f) * OverscanMultiplier) + 1.0f;
+				}
+				LensDistortionHandler->SetOverscanFactor(OverscanFactor);
+
+				//Make sure the displacement map is up to date
+				LensDistortionHandler->ProcessCurrentDistortion();
+
+				bWasDistortionEvaluated = true;
+				break;
+			}
+			default:
+			{
+				ensureMsgf(false, TEXT("Unhandled EDistortionSource type. Update this switch statement."));
+				break; // Do nothing
+			}
+			}
+		}
+
+		if (bApplyDistortion)
+		{
+			if (LensDistortionHandler && bWasDistortionEvaluated)
+			{
+				// Get the current distortion MID from the lens distortion handler
+				UMaterialInstanceDynamic* NewDistortionMID = LensDistortionHandler->GetDistortionMID();
+
+				// If the MID has changed
+				if (LastDistortionMID != NewDistortionMID)
+				{
+					CineCameraComponent->RemoveBlendable(LastDistortionMID);
+					CineCameraComponent->AddOrUpdateBlendable(NewDistortionMID);
+				}
+
+				// Cache the latest distortion MID
+				LastDistortionMID = NewDistortionMID;
+
+				// Get the overscan factor and use it to modify the target camera's FOV
+				const float OverscanFactor = LensDistortionHandler->GetOverscanFactor();
+				const float OverscanSensorWidth = GetDesqueezedSensorWidth(CineCameraComponent) * OverscanFactor;
+				const float OverscanFOV = FMath::RadiansToDegrees(2.0f * FMath::Atan(OverscanSensorWidth / (2.0f * OriginalFocalLength)));
+				CineCameraComponent->SetFieldOfView(OverscanFOV);
+
+				// Update the minimum and maximum focal length of the camera (if needed)
+				CineCameraComponent->LensSettings.MinFocalLength = FMath::Min(CineCameraComponent->LensSettings.MinFocalLength, CineCameraComponent->CurrentFocalLength);
+				CineCameraComponent->LensSettings.MaxFocalLength = FMath::Max(CineCameraComponent->LensSettings.MaxFocalLength, CineCameraComponent->CurrentFocalLength);
+
+				bIsDistortionSetup = true;
+			}
+			else
+			{
+				CleanupDistortion(CineCameraComponent);
+			}
+		}
+		else
+		{
+			// If this is a recorded lens component, and distortion is disabled, set the focal length the original focal length to override the recorded value from the camera.
+			if (EvaluationMode == EFIZEvaluationMode::UseRecordedValues)
+			{
+				CineCameraComponent->CurrentFocalLength = OriginalFocalLength;
+			}
+		}
+
+		LastFocalLength = CineCameraComponent->CurrentFocalLength;
+	}
+
+	bWasLiveLinkFIZUpdated = false;
+
+	Super::TickComponent(DeltaTime, TickType, ThisTickFunction);
+}
+
+void ULensComponent::DestroyComponent(bool bPromoteChildren /*= false*/)
+{
+	Super::DestroyComponent(bPromoteChildren);
+
+	if (UCineCameraComponent* const CineCameraComponent = Cast<UCineCameraComponent>(TargetCameraComponent.GetComponent(GetOwner())))
+	{
+		CleanupDistortion(CineCameraComponent);
+	}
+}
+
+void ULensComponent::PostDuplicate(bool bDuplicateForPIE)
+{
+	Super::PostDuplicate(bDuplicateForPIE);
+
+	InitDefaultCamera();
+}
+
+void ULensComponent::PostEditImport()
+{
+	Super::PostEditImport();
+
+	InitDefaultCamera();
+}
+
+void ULensComponent::PostLoad()
+{
+	Super::PostLoad();
+
+	// After LensModel has been loaded, create a new handler
+	CreateDistortionHandler();
+
+	if (bIsDistortionSetup)
+	{
+		if (UCineCameraComponent* const CineCameraComponent = Cast<UCineCameraComponent>(TargetCameraComponent.GetComponent(GetOwner())))
+		{
+			CleanupDistortion(CineCameraComponent);
+		}
+		bIsDistortionSetup = false;
+	}
+
+	// Cache the currently set LensFile and register for new LensFile events
+	WeakCachedLensFile = LensFilePicker.GetLensFile();
+
+	if (ULensFile* const LensFile = WeakCachedLensFile.Get())
+	{
+		if (!LensFile->OnLensFileModelChanged().IsBoundToObject(this))
+		{
+			LensFile->OnLensFileModelChanged().AddUObject(this, &ULensComponent::OnLensFileModelChanged);
+		}
+	}
+
+#if WITH_EDITOR
+	const int32 UE5MainVersion = GetLinkerCustomVersion(FUE5MainStreamObjectVersion::GUID);
+	if (UE5MainVersion < FUE5MainStreamObjectVersion::LensComponentNodalOffset)
+	{
+		// If this component was previously applying nodal offset, reset the camera component's transform to its original pose
+		if (UCineCameraComponent* const CineCameraComponent = Cast<UCineCameraComponent>(TargetCameraComponent.GetComponent(GetOwner())))
+		{
+			PRAGMA_DISABLE_DEPRECATION_WARNINGS
+			if (bApplyNodalOffsetOnTick)
+			{
+				CineCameraComponent->SetRelativeLocation(OriginalCameraLocation_DEPRECATED);
+				CineCameraComponent->SetRelativeRotation(OriginalCameraRotation_DEPRECATED.Quaternion());
+			}
+			PRAGMA_ENABLE_DEPRECATION_WARNINGS
+		}
+	}
+#endif //WITH_EDITOR
+}
+
+#if WITH_EDITOR
+void ULensComponent::PostEditChangeProperty(struct FPropertyChangedEvent& PropertyChangedEvent)
+{
+	const FName PropertyName = (PropertyChangedEvent.Property != NULL) ? PropertyChangedEvent.Property->GetFName() : NAME_None;
+
+	if (PropertyName == GET_MEMBER_NAME_CHECKED(ULensComponent, bApplyDistortion))
+	{
+		if (UCineCameraComponent* const CineCameraComponent = Cast<UCineCameraComponent>(TargetCameraComponent.GetComponent(GetOwner())))
+		{
+			if (!bApplyDistortion)
+			{
+				CleanupDistortion(CineCameraComponent);
+			}
+		}
+	}
+	else if (PropertyName == GET_MEMBER_NAME_CHECKED(ULensComponent, TargetCameraComponent))
+	{
+		// Clean up distortion on the last target camera
+		if (LastCameraComponent)
+		{
+			CleanupDistortion(LastCameraComponent);
+		}
+
+		LastCameraComponent = Cast<UCineCameraComponent>(TargetCameraComponent.GetComponent(GetOwner()));
+	}
+	else if (PropertyName == GET_MEMBER_NAME_CHECKED(FLensFilePicker, LensFile))
+	{
+		SetLensFilePicker(LensFilePicker);
+	}
+	else if (PropertyName == GET_MEMBER_NAME_CHECKED(ULensComponent, LensModel))
+	{
+		SetLensModel(LensModel);
+	}
+	else if (PropertyName == GET_MEMBER_NAME_CHECKED(ULensComponent, DistortionStateSource))
+	{
+		SetDistortionSource(DistortionStateSource);
+	}
+
+	Super::PostEditChangeProperty(PropertyChangedEvent);
+}
+#endif //WITH_EDITOR
+
+void ULensComponent::ReapplyNodalOffset()
+{
+	// Find the TrackedComponent based on the serialized name
+	TInlineComponentArray<USceneComponent*> SceneComponents;
+	if (AActor* Owner = GetOwner())
+	{
+		Owner->GetComponents(SceneComponents);
+	}
+
+	for (USceneComponent* SceneComponent : SceneComponents)
+	{
+		if (SceneComponent->GetName() == TrackedComponentName)
+		{
+			TrackedComponent = SceneComponent;
+			break;
+		}
+	}
+
+	if (TrackedComponent.IsValid())
+	{
+		// Reset the tracked component back to its original relative transform (before nodal offset was originally applied)
+		TrackedComponent->SetRelativeTransform(OriginalTrackedComponentTransform);
+
+		// Now, reapply the nodal offset to the tracked component
+		ApplyNodalOffset();
+	}
+}
+
+void ULensComponent::ApplyNodalOffset()
+{
+	// Verify that we detected a tracked component for us to offset this tick. If there is none, it is likely because no LiveLink transform controller executed this tick.
+	if (!TrackedComponent.IsValid())
+	{
+		return;
+	}
+
+	// Verify that the evaluation inputs are valid
+	if ((EvaluationMode == EFIZEvaluationMode::DoNotEvaluate) || !EvalInputs.bIsValid) 
+	{
+		return;
+	}
+
+	// Verify that there is a valid LensFile to evaluate
+	ULensFile* LensFile = LensFilePicker.GetLensFile();
+	if (!LensFile)
+	{
+		return;
+	}
+
+	FNodalPointOffset Offset;
+	LensFile->EvaluateNodalPointOffset(EvalInputs.Focus, EvalInputs.Zoom, Offset);
+
+	// Cache the original transform before applying the offset, so that nodal offset can potentially be re-evaluated in the future
+	OriginalTrackedComponentTransform = TrackedComponent.Get()->GetRelativeTransform();
+
+	TrackedComponent.Get()->AddLocalOffset(Offset.LocationOffset);
+	TrackedComponent.Get()->AddLocalRotation(Offset.RotationOffset);
+
+	bWasNodalOffsetAppliedThisTick = true;
+
+	// Reset so that nodal offset will only be applied again next tick if new tracking data was applied between now and then
+	TrackedComponent.Reset();
+}
+
+void ULensComponent::ApplyNodalOffset(USceneComponent* ComponentToOffset, bool bUseManualInputs, float ManualFocusInput, float ManualZoomInput)
+{
+	bApplyNodalOffsetOnTick = false;
+
+	// Verify that the input component was not null
+	if (!ComponentToOffset)
+	{
+		return;
+	}
+
+	// Verify that there is a valid LensFile to evaluate
+	ULensFile* LensFile = LensFilePicker.GetLensFile();
+	if (!LensFile)
+	{
+		return;
+	}
+
+	FNodalPointOffset Offset;
+	if (bUseManualInputs)
+	{
+		LensFile->EvaluateNodalPointOffset(ManualFocusInput, ManualZoomInput, Offset);
+	}
+	else
+	{
+		if (!EvalInputs.bIsValid)
+		{
+			return;
+		}
+
+		LensFile->EvaluateNodalPointOffset(EvalInputs.Focus, EvalInputs.Zoom, Offset);
+	}
+
+	// Cache the original transform before applying the offset, so that nodal offset can potentially be re-evaluated in the future
+	OriginalTrackedComponentTransform = ComponentToOffset->GetRelativeTransform();
+
+	ComponentToOffset->AddLocalOffset(Offset.LocationOffset);
+	ComponentToOffset->AddLocalRotation(Offset.RotationOffset);
+}
+
+void ULensComponent::EvaluateFocalLength(UCineCameraComponent* CineCameraComponent)
+{
+	FFocalLengthInfo FocalLengthInfo;
+	ULensFile* LensFile = LensFilePicker.GetLensFile();
+	if (LensFile && EvalInputs.bIsValid && LensFile->EvaluateFocalLength(EvalInputs.Focus, EvalInputs.Zoom, FocalLengthInfo))
+	{
+		// TODO: EvaluateFocalLength should do this scaling internally
+		const FVector2D FxFyScale = FVector2D(LensFile->LensInfo.SensorDimensions.X / EvalInputs.Filmback.SensorWidth, LensFile->LensInfo.SensorDimensions.Y / EvalInputs.Filmback.SensorHeight);
+		FocalLengthInfo.FxFy = FocalLengthInfo.FxFy * FxFyScale;
+
+		if ((FocalLengthInfo.FxFy.X > UE_KINDA_SMALL_NUMBER) && (FocalLengthInfo.FxFy.Y > UE_KINDA_SMALL_NUMBER))
+		{
+			// This is how field of view, filmback, and focal length are related:
+			//
+			// FOVx = 2*atan(1/(2*Fx)) = 2*atan(FilmbackX / (2*FocalLength))
+			// => FocalLength = Fx*FilmbackX
+			// 
+			// FOVy = 2*atan(1/(2*Fy)) = 2*atan(FilmbackY / (2*FocalLength))
+			// => FilmbackY = FocalLength / Fy
+
+			// Adjust FocalLength and Filmback to match FxFy (which has already been divided by resolution in pixels)
+			const float NewFocalLength = FocalLengthInfo.FxFy.X * EvalInputs.Filmback.SensorWidth * LensFile->LensInfo.SqueezeFactor;
+
+			// TODO: Consider adding an advanced setting to control whether or not sensor height should be modified (if Fx/Fy don't have a fixed aspect ratio)
+
+			// Update the minimum and maximum focal length of the camera (if needed)
+			CineCameraComponent->LensSettings.MinFocalLength = FMath::Max(FMath::Min(CineCameraComponent->LensSettings.MinFocalLength, NewFocalLength), 0.01);
+			CineCameraComponent->LensSettings.MaxFocalLength = FMath::Max(FMath::Max(CineCameraComponent->LensSettings.MaxFocalLength, NewFocalLength), 0.01);
+
+			CineCameraComponent->SetCurrentFocalLength(NewFocalLength);
+
+			OriginalFocalLength = NewFocalLength;
+		}
+	}
+}
+
+FLensFilePicker ULensComponent::GetLensFilePicker() const
+{
+	return LensFilePicker;
+}
+
+ULensFile* ULensComponent::GetLensFile() const
+{
+	return LensFilePicker.GetLensFile();
+}
+
+void ULensComponent::SetLensFilePicker(FLensFilePicker LensFile)
+{
+	LensFilePicker = LensFile;
+
+	if (ULensFile* const CachedLensFile = WeakCachedLensFile.Get())
+	{
+		CachedLensFile->OnLensFileModelChanged().RemoveAll(this);
+	}
+
+	WeakCachedLensFile = LensFilePicker.GetLensFile();
+	if (ULensFile* const NewLensFile = WeakCachedLensFile.Get())
+	{
+		if (!NewLensFile->OnLensFileModelChanged().IsBoundToObject(this))
+		{
+			NewLensFile->OnLensFileModelChanged().AddUObject(this, &ULensComponent::OnLensFileModelChanged);
+		}
+	}
+
+	if (DistortionStateSource == EDistortionSource::LensFile)
+	{
+		if (ULensFile* const NewLensFile = LensFilePicker.GetLensFile())
+		{
+			SetLensModel(NewLensFile->LensInfo.LensModel);
+		}
+		else
+		{
+			ClearDistortionState();
+		}
+	}
+}
+
+void ULensComponent::SetLensFile(ULensFile* Lens)
+{
+	// Automatically sets this to false so the component can use the newly set LensFile directly
+	FLensFilePicker NewLensFilePicker;
+	NewLensFilePicker.bUseDefaultLensFile = false;
+	NewLensFilePicker.LensFile = Lens;
+
+	SetLensFilePicker(NewLensFilePicker);
+}
+
+EFIZEvaluationMode ULensComponent::GetFIZEvaluationMode() const
+{
+	return EvaluationMode;
+}
+
+void ULensComponent::SetFIZEvaluationMode(EFIZEvaluationMode Mode)
+{
+	EvaluationMode = Mode;
+}
+
+float ULensComponent::GetOverscanMultiplier() const
+{
+	if (bScaleOverscan)
+	{
+		return OverscanMultiplier;
+	}
+	return 1.0f;
+}
+
+void ULensComponent::SetOverscanMultiplier(float Multiplier)
+{
+	bScaleOverscan = true;
+	OverscanMultiplier = Multiplier;
+}
+
+EFilmbackOverrideSource ULensComponent::GetFilmbackOverrideSetting() const
+{
+	return FilmbackOverride;
+}
+
+void ULensComponent::SetFilmbackOverrideSetting(EFilmbackOverrideSource Setting)
+{
+	FilmbackOverride = Setting;
+}
+
+FCameraFilmbackSettings ULensComponent::GetCroppedFilmback() const
+{
+	return CroppedFilmback;
+}
+
+void ULensComponent::SetCroppedFilmback(FCameraFilmbackSettings Filmback)
+{
+	CroppedFilmback = Filmback;
+}
+
+const FLensFileEvaluationInputs& ULensComponent::GetLensFileEvaluationInputs() const 
+{
+	return EvalInputs;
+}
+
+bool ULensComponent::ShouldApplyNodalOffsetOnTick() const
+{
+	return bApplyNodalOffsetOnTick;
+}
+
+void ULensComponent::SetApplyNodalOffsetOnTick(bool bApplyNodalOffset)
+{
+	bApplyNodalOffsetOnTick = bApplyNodalOffset;
+}
+
+bool ULensComponent::WasNodalOffsetAppliedThisTick() const 
+{
+	return bWasNodalOffsetAppliedThisTick;
+}
+
+bool ULensComponent::WasDistortionEvaluated() const
+{
+	return bWasDistortionEvaluated;
+}
+
+float ULensComponent::GetOriginalFocalLength() const 
+{
+	return OriginalFocalLength;
+}
+
+TSubclassOf<ULensModel> ULensComponent::GetLensModel() const
+{
+	return LensModel;
+}
+
+void ULensComponent::SetLensModel(TSubclassOf<ULensModel> Model)
+{
+	LensModel = Model;
+
+	if (LensModel)
+	{
+		LensModel->GetDefaultObject<ULensModel>()->GetDefaultParameterArray(DistortionState.DistortionInfo.Parameters);
+
+		CreateDistortionHandler();
+	}
+	else
+	{
+		DistortionState.DistortionInfo.Parameters.Empty();
+	}
+
+	OnLensComponentModelChangedDelegate.Broadcast(LensModel);
+}
+
+void ULensComponent::OnLensFileModelChanged(const TSubclassOf<ULensModel>& Model)
+{
+	if (DistortionStateSource == EDistortionSource::LensFile)
+	{
+		SetLensModel(Model);
+	}
+}
+
+FLensDistortionState ULensComponent::GetDistortionState() const
+{
+	return DistortionState;
+}
+
+void ULensComponent::SetDistortionState(FLensDistortionState State)
+{
+	DistortionState = State;
+}
+
+EDistortionSource ULensComponent::GetDistortionSource() const
+{
+	return DistortionStateSource;
+}
+
+void ULensComponent::SetDistortionSource(EDistortionSource Source)
+{
+	DistortionStateSource = Source;
+
+	// If the new source is a LensFile, update the lens model to match
+	if (DistortionStateSource == EDistortionSource::LensFile)
+	{
+		if (ULensFile* LensFile = LensFilePicker.GetLensFile())
+		{
+			SetLensModel(LensFile->LensInfo.LensModel);
+		}
+	}
+
+	ClearDistortionState();
+}
+
+bool ULensComponent::ShouldApplyDistortion() const
+{
+	return bApplyDistortion;
+}
+
+void ULensComponent::SetApplyDistortion(bool bApply)
+{
+	bApplyDistortion = bApply;
+}
+
+ULensDistortionModelHandlerBase* ULensComponent::GetLensDistortionHandler() const
+{
+	return LensDistortionHandlerMap.FindRef(LensModel);
+}
+
+PRAGMA_DISABLE_DEPRECATION_WARNINGS
+FDistortionHandlerPicker ULensComponent::GetDistortionHandlerPicker() const
+{
+#if WITH_EDITOR
+	return DistortionSource_DEPRECATED;
+#else
+	return FDistortionHandlerPicker();
+#endif //WITH_EDITOR
+}
+PRAGMA_ENABLE_DEPRECATION_WARNINGS
+
+void ULensComponent::InitDefaultCamera()
+{
+	UCineCameraComponent* TargetCamera = Cast<UCineCameraComponent>(TargetCameraComponent.GetComponent(GetOwner()));
+	if (!TargetCamera)
+	{
+		// Find all CineCameraComponents on the same actor as this component and set the first one to be the target
+		TInlineComponentArray<UCineCameraComponent*> CineCameraComponents;
+		if (AActor* Owner = GetOwner())
+		{
+			Owner->GetComponents(CineCameraComponents);
+		}
+
+		if (CineCameraComponents.Num() > 0)
+		{
+			TargetCameraComponent.ComponentProperty = CineCameraComponents[0]->GetFName();
+			LastCameraComponent = CineCameraComponents[0];
+		}
+	}
+}
+
+float ULensComponent::GetDesqueezedSensorWidth(UCineCameraComponent* const CineCameraComponent) const
+{
+	if (CineCameraComponent)
+	{
+		return CineCameraComponent->Filmback.SensorWidth * CineCameraComponent->LensSettings.SqueezeFactor;
+	}
+
+	return 1.0f;
+}
+
+void ULensComponent::CleanupDistortion(UCineCameraComponent* const CineCameraComponent)
+{
+	if (bIsDistortionSetup)
+	{
+		// Remove the last distortion MID that was applied to the target camera component
+		if (LastDistortionMID)
+		{
+			CineCameraComponent->RemoveBlendable(LastDistortionMID);
+			LastDistortionMID = nullptr;
+		}
+
+		// Restore the original FOV of the target camera
+		const float UndistortedFOV = FMath::RadiansToDegrees(2.0f * FMath::Atan(GetDesqueezedSensorWidth(CineCameraComponent) / (2.0f * OriginalFocalLength)));
+		CineCameraComponent->SetFieldOfView(UndistortedFOV);
+
+		// Update the minimum and maximum focal length of the camera (if needed)
+		CineCameraComponent->LensSettings.MinFocalLength = FMath::Min(CineCameraComponent->LensSettings.MinFocalLength, CineCameraComponent->CurrentFocalLength);
+		CineCameraComponent->LensSettings.MaxFocalLength = FMath::Max(CineCameraComponent->LensSettings.MaxFocalLength, CineCameraComponent->CurrentFocalLength);
+	}
+
+	bIsDistortionSetup = false;
+}
+
+void ULensComponent::ClearDistortionState()
+{
+	FLensDistortionState ZeroDistortionState;
+	if (LensModel)
+	{
+		LensModel->GetDefaultObject<ULensModel>()->GetDefaultParameterArray(ZeroDistortionState.DistortionInfo.Parameters);
+	}
+	DistortionState = ZeroDistortionState;
+}
+
+void ULensComponent::CreateDistortionHandler()
+{
+	if (LensModel)
+	{
+		TObjectPtr<ULensDistortionModelHandlerBase>& Handler = LensDistortionHandlerMap.FindOrAdd(LensModel);
+		if (!Handler)
+		{
+			if (const TSubclassOf<ULensDistortionModelHandlerBase> HandlerClass = ULensModel::GetHandlerClass(LensModel))
+			{
+				Handler = NewObject<ULensDistortionModelHandlerBase>(this, HandlerClass);
+			}
+		}
+	}
+}
+
+void ULensComponent::OnLiveLinkComponentRegistered(ULiveLinkComponentController* LiveLinkComponent)
+{
+	// Check that the new LiveLinkComponent that was just registered was added to the same actor that this component belongs to
+	if (LiveLinkComponent->GetOwner() == GetOwner())
+	{
+		AddTickPrerequisiteComponent(LiveLinkComponent);
+
+		if (!LiveLinkComponent->OnLiveLinkControllersTicked().IsBoundToObject(this))
+		{
+			LiveLinkComponent->OnLiveLinkControllersTicked().AddUObject(this, &ULensComponent::ProcessLiveLinkData);
+		}
+	}
+}
+
+void ULensComponent::ProcessLiveLinkData(const ULiveLinkComponentController* const LiveLinkComponent, const FLiveLinkSubjectFrameData& SubjectData)
+{
+	TSubclassOf<ULiveLinkRole> LiveLinkRole = LiveLinkComponent->GetSubjectRepresentation().Role;
+
+	if (LiveLinkRole.Get()->IsChildOf(ULiveLinkCameraRole::StaticClass()))
+	{
+		UpdateLiveLinkFIZ(LiveLinkComponent, SubjectData);
+	}
+
+	if (LiveLinkRole.Get()->IsChildOf(ULiveLinkTransformRole::StaticClass()))
+ 	{
+ 		UpdateTrackedComponent(LiveLinkComponent, SubjectData);
+ 	}
+}
+
+void ULensComponent::UpdateTrackedComponent(const ULiveLinkComponentController* const LiveLinkComponent, const FLiveLinkSubjectFrameData& SubjectData)
+{
+	const FLiveLinkTransformStaticData* StaticData = SubjectData.StaticData.Cast<FLiveLinkTransformStaticData>();
+	const FLiveLinkTransformFrameData* FrameData = SubjectData.FrameData.Cast<FLiveLinkTransformFrameData>();
+
+	check(StaticData && FrameData);
+
+	// Find the transform controller in the LiveLink component's controller map
+	if (const TObjectPtr<ULiveLinkControllerBase>* TransformControllerPtr = LiveLinkComponent->ControllerMap.Find(ULiveLinkTransformRole::StaticClass()))
+	{
+		if (ULiveLinkTransformController* TransformController = Cast<ULiveLinkTransformController>(*TransformControllerPtr))
+		{
+			// Check LiveLink static data to ensure that location and rotation were supported by this subject
+			bool bIsLocationRotationSupported = (StaticData->bIsLocationSupported && StaticData->bIsRotationSupported);
+
+			// Check the transform controller usage flags to ensure that location and rotation were supported by the controller
+			bIsLocationRotationSupported = bIsLocationRotationSupported && (TransformController->TransformData.bUseLocation && TransformController->TransformData.bUseRotation);
+
+			if (bIsLocationRotationSupported)
+			{
+				TrackedComponent = Cast<USceneComponent>(TransformController->GetAttachedComponent());
+				if (TrackedComponent.IsValid())
+				{
+					TrackedComponentName = TrackedComponent.Get()->GetName();
+				}
+				// At this point, we know that the tracked component's transform was just set by LiveLink, and therefore has no nodal offset applied
+				bWasNodalOffsetAppliedThisTick = false;
+			}
+		}
+	}
+}
+
+void ULensComponent::UpdateLiveLinkFIZ(const ULiveLinkComponentController* const LiveLinkComponent, const FLiveLinkSubjectFrameData& SubjectData)
+{
+	const FLiveLinkCameraStaticData* StaticData = SubjectData.StaticData.Cast<FLiveLinkCameraStaticData>();
+	const FLiveLinkCameraFrameData* FrameData = SubjectData.FrameData.Cast<FLiveLinkCameraFrameData>();
+	
+	check(StaticData && FrameData);
+
+	// Check that the camera component being controlled by the input LiveLink component matches our target camera
+	if (const TObjectPtr<ULiveLinkControllerBase>* CameraControllerPtr = LiveLinkComponent->ControllerMap.Find(ULiveLinkCameraRole::StaticClass()))
+	{
+		if (ULiveLinkControllerBase* CameraController = *CameraControllerPtr)
+		{
+			if (CameraController->GetAttachedComponent() != TargetCameraComponent.GetComponent(GetOwner()))
+			{
+				return;
+			}
+		}
+	}
+
+	if (StaticData->bIsFocusDistanceSupported && StaticData->bIsFocalLengthSupported)
+	{
+		LiveLinkFocus = FrameData->FocusDistance;
+		LiveLinkZoom = FrameData->FocalLength;
+
+		LiveLinkIris = FrameData->Aperture;
+
+		bWasLiveLinkFIZUpdated = true;
+	}
+}
+
+void ULensComponent::UpdateLensFileEvaluationInputs(UCineCameraComponent* CineCameraComponent)
+{
+	EvalInputs.bIsValid = false;
+
+	// Query for the focus and zoom inputs to use to evaluate the LensFile. The source of these inputs will depend on the evaluation mode.
+	if (EvaluationMode == EFIZEvaluationMode::UseLiveLink)
+	{
+		if (bWasLiveLinkFIZUpdated)
+		{
+			EvalInputs.Focus = LiveLinkFocus;
+			EvalInputs.Iris = LiveLinkIris;
+			EvalInputs.Zoom = LiveLinkZoom;
+			EvalInputs.bIsValid = true;
+		}
+	}
+	else if (EvaluationMode == EFIZEvaluationMode::UseCameraSettings)
+	{
+		EvalInputs.Focus = CineCameraComponent->CurrentFocusDistance;
+		EvalInputs.Iris = CineCameraComponent->CurrentAperture;
+		// The camera's focal length might have been altered by applying an overscan FOV. However, the LensFile needs to be evaluated with the non-overscanned value
+		EvalInputs.Zoom = OriginalFocalLength;
+		EvalInputs.bIsValid = true;
+	}
+	else if (EvaluationMode == EFIZEvaluationMode::UseRecordedValues)
+	{
+		// Do nothing, the values for EvalInputs.Focus and EvalInputs.Zoom are already loaded from the recorded sequence
+		EvalInputs.bIsValid = true;
+	}
+}
+
+void ULensComponent::UpdateCameraFilmback(UCineCameraComponent* CineCameraComponent)
+{	
+	switch (FilmbackOverride)
+	{
+		case EFilmbackOverrideSource::LensFile:
+		{	
+			if (ULensFile* LensFile = LensFilePicker.GetLensFile())
+			{
+				CineCameraComponent->Filmback.SensorWidth = LensFile->LensInfo.SensorDimensions.X;
+				CineCameraComponent->Filmback.SensorHeight = LensFile->LensInfo.SensorDimensions.Y;
+				CineCameraComponent->LensSettings.SqueezeFactor = LensFile->LensInfo.SqueezeFactor;
+			}
+			break;
+		}
+		case EFilmbackOverrideSource::CroppedFilmbackSetting:
+		{
+			CineCameraComponent->Filmback = CroppedFilmback;
+
+			if (ULensFile* LensFile = LensFilePicker.GetLensFile())
+			{
+				CineCameraComponent->LensSettings.SqueezeFactor = LensFile->LensInfo.SqueezeFactor;
+			}
+			break;
+		}
+		case EFilmbackOverrideSource::DoNotOverride:
+		{
+			break; // Do Nothing
+		}
+		default:
+		{
+			ensureMsgf(false, TEXT("Unhandled EFilmbackOverrideSource type. Update this switch statement."));
+			break; 
+		}
+	}
+
+	EvalInputs.Filmback = CineCameraComponent->Filmback;
+}
diff --git a/Engine/Plugins/VirtualProduction/LensComponent/Source/LensComponent/Private/LensComponentModule.cpp b/Engine/Plugins/VirtualProduction/LensComponent/Source/LensComponent/Private/LensComponentModule.cpp
new file mode 100644
--- /dev/null
+++ b/Engine/Plugins/VirtualProduction/LensComponent/Source/LensComponent/Private/LensComponentModule.cpp
@@ -0,0 +1,5 @@
+// Copyright Epic Games, Inc. All Rights Reserved.
+
+#include "Modules/ModuleManager.h"
+
+IMPLEMENT_MODULE(FDefaultModuleImpl, LensComponent);
diff --git a/Engine/Plugins/VirtualProduction/LensComponent/Source/LensComponent/Public/LensComponent.h b/Engine/Plugins/VirtualProduction/LensComponent/Source/LensComponent/Public/LensComponent.h
new file mode 100644
--- /dev/null
+++ b/Engine/Plugins/VirtualProduction/LensComponent/Source/LensComponent/Public/LensComponent.h
@@ -0,0 +1,387 @@
+// Copyright Epic Games, Inc. All Rights Reserved.
+
+#pragma once
+
+#include "Components/ActorComponent.h"
+
+#include "CameraCalibrationTypes.h"
+#include "LensFile.h"
+#include "LiveLinkComponentController.h"
+
+#include "LensComponent.generated.h"
+
+class FLensComponentDetailCustomization;
+class UCineCameraComponent;
+
+
+DECLARE_MULTICAST_DELEGATE_OneParam(FOnLensComponentModelChanged, const TSubclassOf<ULensModel>&);
+
+
+/** Mode that controls where FIZ inputs are sourced from and how they are used to evaluate the LensFile */
+UENUM(BlueprintType)
+enum class EFIZEvaluationMode : uint8
+{
+	/** Evaluate the Lens File with the latest FIZ data received from LiveLink */
+	UseLiveLink,
+	/** Evaluate the Lens File using the current FIZ settings of the target camera */
+	UseCameraSettings,
+	/** Evaluate the Lens File using values recorded in a level sequence (set automatically when the sequence is opened) */
+	UseRecordedValues,
+	/** Do not evaluate the Lens File */
+	DoNotEvaluate,
+};
+
+/** Controls whether this component can override the camera's filmback, and if so, which override to use */
+UENUM(BlueprintType)
+enum class EFilmbackOverrideSource : uint8
+{
+	/** Override the camera's filmback using the sensor dimensions recorded in the LensInfo of the LensFile */
+	LensFile,
+	/** Override the camera's filmback using the CroppedFilmback setting below */
+	CroppedFilmbackSetting,
+	/** Do not override the camera's filmback */
+	DoNotOverride,
+};
+
+/** Specifies from where the distortion state information comes */
+UENUM(BlueprintType)
+enum class EDistortionSource : uint8
+{
+	/** Distortion state is evaluated using the LensFile */
+	LensFile,
+	/** Distortion state is inputted directly from a LiveLink subject */
+	LiveLinkLensSubject,
+	/** Distortion state is set manually by the user using the Distortion State setting below */
+	Manual,
+};
+
+/** Component for applying a post-process lens distortion effect to a CineCameraComponent on the same actor */
+UCLASS(HideCategories=(Tags, Activation, Cooking, AssetUserData, Collision), meta=(BlueprintSpawnableComponent))
+class LENSCOMPONENT_API ULensComponent : public UActorComponent
+{
+	GENERATED_BODY()
+
+	friend FLensComponentDetailCustomization;
+
+public:
+	ULensComponent();
+
+	//~ Begin UActorComponent interface
+	virtual void OnRegister() override;
+	virtual void OnUnregister() override;
+	virtual void TickComponent(float DeltaTime, ELevelTick TickType, FActorComponentTickFunction* ThisTickFunction) override;
+	virtual void DestroyComponent(bool bPromoteChildren = false) override;
+	//~ End UActorComponent interface
+
+	//~ Begin UObject interface
+	virtual void PostDuplicate(bool bDuplicateForPIE) override;
+	virtual void PostEditImport() override;
+	virtual void PostLoad() override;
+
+#if WITH_EDITOR
+	virtual void PostEditChangeProperty(struct FPropertyChangedEvent& PropertyChangedEvent) override;
+#endif //WITH_EDITOR
+	//~ End UObject interface
+
+	/** Returns the delegate that is triggered when the LensModel changes */
+	FOnLensComponentModelChanged& OnLensComponentModelChanged() 
+	{ 
+		return OnLensComponentModelChangedDelegate; 
+	}
+
+	/** Get the LensFile picker used by this component */
+	UFUNCTION(BlueprintPure, Category = "Lens Component")
+	FLensFilePicker GetLensFilePicker() const;
+
+	/** Get the LensFile used by this component */
+	UFUNCTION(BlueprintPure, Category="Lens Component")
+	ULensFile* GetLensFile() const;
+
+	/** Set the LensFile picker used by this component */
+	UFUNCTION(BlueprintCallable, Category = "Lens Component")
+	void SetLensFilePicker(FLensFilePicker LensFile);
+
+	/** Set the LensFile used by this component */
+	UFUNCTION(BlueprintCallable, Category="Lens Component")
+	void SetLensFile(ULensFile* LensFile);
+
+	/** Get the evaluation mode used to evaluate the LensFile */
+	UFUNCTION(BlueprintPure, Category="Lens Component", meta = (DisplayName = "Get FIZ Evaluation Mode"))
+	EFIZEvaluationMode GetFIZEvaluationMode() const;
+
+	/** Set the evaluation mode used to evaluate the LensFile */
+	UFUNCTION(BlueprintCallable, Category="Lens Component", meta=(DisplayName="Set FIZ Evaluation Mode"))
+	void SetFIZEvaluationMode(EFIZEvaluationMode Mode);
+
+	/** Get the evaluation mode used to evaluate the LensFile */
+	UFUNCTION(BlueprintPure, Category = "Lens Component")
+	float GetOverscanMultiplier() const;
+
+	/** Set the LensFile used by this component */
+	UFUNCTION(BlueprintCallable, Category = "Lens Component")
+	void SetOverscanMultiplier(float Multiplier);
+
+	/** Get the filmback override setting */
+	UFUNCTION(BlueprintPure, Category = "Lens Component")
+	EFilmbackOverrideSource GetFilmbackOverrideSetting() const;
+
+	/** Set the filmback override setting */
+	UFUNCTION(BlueprintCallable, Category = "Lens Component")
+	void SetFilmbackOverrideSetting(EFilmbackOverrideSource Setting);
+
+	/** Get the cropped filmback (only relevant if the filmback override setting is set to use the CroppedFilmback */
+	UFUNCTION(BlueprintPure, Category = "Lens Component")
+	FCameraFilmbackSettings GetCroppedFilmback() const;
+
+	/** Set the cropped filmback (only relevant if the filmback override setting is set to use the CroppedFilmback */
+	UFUNCTION(BlueprintCallable, Category = "Lens Component")
+	void SetCroppedFilmback(FCameraFilmbackSettings Filmback);
+
+	/** Returns true if nodal offset will be automatically applied during this component's tick, false otherwise */
+	UFUNCTION(BlueprintPure, Category="Lens Component")
+	bool ShouldApplyNodalOffsetOnTick() const;
+
+	/** Set whether nodal offset should be automatically applied during this component's tick */
+	UFUNCTION(BlueprintCallable, Category="Lens Component")
+	void SetApplyNodalOffsetOnTick(bool bApplyNodalOffset);
+
+	/** Get the distortion source setting */
+	UFUNCTION(BlueprintPure, Category = "Lens Component")
+	EDistortionSource GetDistortionSource() const;
+
+	/** Set the distortion source setting */
+	UFUNCTION(BlueprintCallable, Category = "Lens Component")
+	void SetDistortionSource(EDistortionSource Source);
+
+	/** Whether distortion should be applied to the target camera */
+	UFUNCTION(BlueprintPure, Category = "Lens Component")
+	bool ShouldApplyDistortion() const;
+
+	/** Set whether distortion should be applied to the target camera */
+	UFUNCTION(BlueprintCallable, Category = "Lens Component")
+	void SetApplyDistortion(bool bApply);
+
+	/** Get the current lens model */
+	UFUNCTION(BlueprintPure, Category = "Lens Component")
+	TSubclassOf<ULensModel> GetLensModel() const;
+
+	/** Set the current lens model */
+	UFUNCTION(BlueprintCallable, Category = "Lens Component")
+	void SetLensModel(TSubclassOf<ULensModel> Model);
+
+	/** Get the current distortion state */
+	UFUNCTION(BlueprintPure, Category = "Lens Component")
+	FLensDistortionState GetDistortionState() const;
+
+	/** Set the current distortion state */
+	UFUNCTION(BlueprintCallable, Category = "Lens Component")
+	void SetDistortionState(FLensDistortionState State);
+
+	/** Reset the distortion state to defaults to represent "no distortion" */
+	UFUNCTION(BlueprintCallable, Category = "Lens Component")
+	void ClearDistortionState();
+
+	/** Get the original (not adjusted for overscan) focal length of the camera */
+	UFUNCTION(BlueprintPure, Category = "Lens Component")
+	float GetOriginalFocalLength() const;
+
+	/** Get the data used by this component to evaluate the LensFile */
+	UFUNCTION(BlueprintPure, Category = "Lens Component")
+	const FLensFileEvaluationInputs& GetLensFileEvaluationInputs() const;
+
+	/** Returns true if nodal offset was applied during the current tick, false otherwise */
+	UFUNCTION(BlueprintPure, Category = "Lens Component")
+	bool WasNodalOffsetAppliedThisTick() const;
+
+	/** Returns true if distortion was evaluated this tick */
+	UFUNCTION(BlueprintPure, Category = "Lens Component")
+	bool WasDistortionEvaluated() const;
+
+	/** 
+	 * Manually apply nodal offset to the specified component. 
+	 * If bUseManualInputs is true, the input Focus and Zoom values will be used to evaluate the LensFile .
+	 * If bUseManualInputs is false, the LensFile be will evaluated based on the Lens Component's evaluation mode.
+	 */
+	UFUNCTION(BlueprintCallable, Category="Lens Component", meta=(AdvancedDisplay=1))
+	void ApplyNodalOffset(USceneComponent* ComponentToOffset, bool bUseManualInputs = false, float ManualFocusInput = 0.0f, float ManualZoomInput = 0.0f);
+
+public:
+	/** Returns the LensDistortionHandler in use for the current LensModel */
+	ULensDistortionModelHandlerBase* GetLensDistortionHandler() const;
+
+	/** Reset the tracked component back to its original tracked pose and reapply nodal offset to it by re-evaluating the LensFile */
+	void ReapplyNodalOffset();
+
+PRAGMA_DISABLE_DEPRECATION_WARNINGS
+	UE_DEPRECATED(5.1, "This property is deprecated. Use GetDistortionSource() and SetDistortionSource() instead.")
+	FDistortionHandlerPicker GetDistortionHandlerPicker() const;
+PRAGMA_ENABLE_DEPRECATION_WARNINGS
+
+private:
+	/** Evaluate the LensFile for nodal offset (using the current evaluation mode) and apply it to the latest component to offset */
+	void ApplyNodalOffset();
+
+	/** Evaluate the focal length from the LensFile and applies the calibrated value to the camera */
+	void EvaluateFocalLength(UCineCameraComponent* CineCameraComponent);
+
+	/** If TargetCameraComponent is not set, initialize it to the first CineCameraComponent on the same actor as this component */
+	void InitDefaultCamera();
+
+	/** Remove the last distortion MID applied to the input CineCameraComponent and reset its FOV to use no overscan */
+	void CleanupDistortion(UCineCameraComponent* CineCameraComponent);
+
+	/** Register a new lens distortion handler with the camera calibration subsystem using the selected lens file */
+	void CreateDistortionHandler();
+
+	/** Register to the new LiveLink component's callback to be notified when its controller map changes */
+	void OnLiveLinkComponentRegistered(ULiveLinkComponentController* LiveLinkComponent);
+	
+	/** Triggered when the LensFile model changes */
+	void OnLensFileModelChanged(const TSubclassOf<ULensModel>& Model);
+
+	/** Callback executed when a LiveLink component on the same actor ticks */
+	void ProcessLiveLinkData(const ULiveLinkComponentController* const LiveLinkComponent, const FLiveLinkSubjectFrameData& SubjectData);
+
+	/** Inspects the subject data and LiveLink transform controller to determine which component (if any) had tracking data applied to it */
+	void UpdateTrackedComponent(const ULiveLinkComponentController* const LiveLinkComponent, const FLiveLinkSubjectFrameData& SubjectData);
+
+	/** Inspects the subject data and LiveLink camera controller cache the FIZ that was input for the target camera */
+	void UpdateLiveLinkFIZ(const ULiveLinkComponentController* const LiveLinkComponent, const FLiveLinkSubjectFrameData& SubjectData);
+
+	/** Updates the focus and zoom inputs that will be used to evaluate the LensFile based on the evaluation mode */
+	void UpdateLensFileEvaluationInputs(UCineCameraComponent* CineCameraComponent);
+
+	/** Updates the camera's filmback based on the filmback override settings */
+	void UpdateCameraFilmback(UCineCameraComponent* CineCameraComponent);
+
+	/** Returns the sensor width of the input CineCamera, factoring in its squeeze factor */
+	float GetDesqueezedSensorWidth(UCineCameraComponent* const CineCameraComponent) const;
+
+protected:
+	/** Lens File used to drive distortion with current camera settings */
+	UPROPERTY(EditAnywhere, Category="Lens File", meta=(ShowOnlyInnerProperties))
+	FLensFilePicker LensFilePicker;
+
+	/** Specify how the Lens File should be evaluated */
+	UPROPERTY(EditAnywhere, Category="Lens File")
+	EFIZEvaluationMode EvaluationMode = EFIZEvaluationMode::UseLiveLink;
+
+	/** The CineCameraComponent on which to apply the post-process distortion effect */
+	UPROPERTY(EditInstanceOnly, AdvancedDisplay, Category="Lens File", meta=(UseComponentPicker, AllowedClasses="/Script/CinematicCamera.CineCameraComponent"))
+	FComponentReference TargetCameraComponent;
+
+	/** Inputs to LensFile evaluation */
+	UPROPERTY(VisibleAnywhere, AdvancedDisplay, Category = "Lens File")
+	FLensFileEvaluationInputs EvalInputs;
+
+	/** Specifies from where the distortion state information comes */
+	UPROPERTY(EditAnywhere, Category = "Distortion", meta = (DisplayName = "Distortion Source"))
+	EDistortionSource DistortionStateSource = EDistortionSource::LensFile;
+
+	/** Whether or not to apply distortion to the target camera component */
+	UPROPERTY(EditAnywhere, BlueprintReadWrite, Getter = ShouldApplyDistortion, Setter = SetApplyDistortion, Category = "Distortion")
+	bool bApplyDistortion = false;
+
+	/** The current lens model used for distortion */
+	UPROPERTY(EditAnywhere, Category = "Distortion", meta = (EditCondition = "DistortionStateSource == EDistortionSource::Manual"))
+	TSubclassOf<ULensModel> LensModel;
+
+	/** The current distortion state */
+	UPROPERTY(Interp, EditAnywhere, Category = "Distortion", meta = (EditCondition = "DistortionStateSource == EDistortionSource::Manual"))
+	FLensDistortionState DistortionState;
+
+	/** Whether to scale the computed overscan by the overscan percentage */
+	UPROPERTY(AdvancedDisplay, BlueprintReadWrite, Category = "Distortion")
+	bool bScaleOverscan = false;
+
+	/** The percentage of the computed overscan that should be applied to the target camera */
+	UPROPERTY(EditAnywhere, AdvancedDisplay, BlueprintReadWrite, Category = "Distortion", meta = (EditCondition = "bScaleOverscan", ClampMin = "0.0", ClampMax = "2.0"))
+	float OverscanMultiplier = 1.0f;
+
+	/** Controls whether this component can override the camera's filmback, and if so, which override to use */
+	UPROPERTY(EditAnywhere, Category = "Filmback")
+	EFilmbackOverrideSource FilmbackOverride = EFilmbackOverrideSource::DoNotOverride;
+
+	/** Cropped filmback to use if the filmback override settings are set to use it */
+	UPROPERTY(EditAnywhere, Category = "Filmback", meta=(EditCondition="FilmbackOverride == EFilmbackOverrideSource::CroppedFilmbackSetting"))
+	FCameraFilmbackSettings CroppedFilmback;
+
+	/** 
+	 * If checked, nodal offset will be applied automatically when this component ticks. 
+	 * Set to false if nodal offset needs to be manually applied at some other time (via Blueprints).
+	 */
+	UPROPERTY(EditAnywhere, Category="Nodal Offset")
+	bool bApplyNodalOffsetOnTick = true;
+
+	/** Serialized transform of the TrackedComponent prior to nodal offset being applied */
+	UPROPERTY(Interp, VisibleAnywhere, AdvancedDisplay, Category = "Nodal Offset")
+	FTransform OriginalTrackedComponentTransform;
+
+	/** Whether a distortion effect is currently being applied to the target camera component */
+	UPROPERTY()
+	bool bIsDistortionSetup = false;
+
+	/** Focal length of the target camera before any overscan has been applied */
+	UPROPERTY()
+	float OriginalFocalLength = 35.0f;
+
+	/** Cached MID last applied to the target camera */
+	UPROPERTY()
+	TObjectPtr<UMaterialInstanceDynamic> LastDistortionMID = nullptr;
+
+	/** Cached most recent target camera, used to clean up the old camera when the user changes the target */
+	UPROPERTY()
+	TObjectPtr<UCineCameraComponent> LastCameraComponent = nullptr;
+
+	/** Map of lens models to handlers */
+	UPROPERTY(Transient)
+	TMap<TSubclassOf<ULensModel>, TObjectPtr<ULensDistortionModelHandlerBase>> LensDistortionHandlerMap;
+
+	/** Scene component that should have nodal offset applied */
+	UPROPERTY(Transient)
+	TWeakObjectPtr<USceneComponent> TrackedComponent;
+
+	/** Serialized name of the TrackedComponent, used to determine which component to re-apply nodal offset to in spawnables */
+	UPROPERTY()
+	FString TrackedComponentName;
+
+	/** Weak pointer to the currently set LensFile. Useful to unregister events from the previous LensFile when the asset is swapped out */
+	TWeakObjectPtr<ULensFile> WeakCachedLensFile;
+
+#if WITH_EDITORONLY_DATA
+	UE_DEPRECATED(5.1, "This property has been deprecated. The LensDistortion component no longer tracks the attached camera's original rotation.")
+	UPROPERTY()
+	FRotator OriginalCameraRotation_DEPRECATED;
+
+	UE_DEPRECATED(5.1, "This property has been deprecated. The LensDistortion component no longer tracks the attached camera's original location.")
+	UPROPERTY()
+	FVector OriginalCameraLocation_DEPRECATED;
+
+PRAGMA_DISABLE_DEPRECATION_WARNINGS
+	UE_DEPRECATED(5.1, "This property has been deprecated. The handler picker is no longer used to identify a distortion handler. Use the LensDistortionHandlerMap to get a handler for the current LensModel.")
+	UPROPERTY(meta=(DeprecatedProperty, DeprecationMessage="This property has been deprecated. Use GetDistortionSource() and SetDistortionSource() instead."))
+	FDistortionHandlerPicker DistortionSource_DEPRECATED;
+PRAGMA_ENABLE_DEPRECATION_WARNINGS
+#endif //WITH_EDITORONLY_DATA
+
+private:
+	/** Latest LiveLink FIZ data, used to evaluate the LensFile */
+	float LiveLinkFocus = 0.0f;
+	float LiveLinkIris = 0.0f;
+	float LiveLinkZoom = 0.0f;
+
+	/** Whether LiveLink FIZ was received this tick */
+	bool bWasLiveLinkFIZUpdated = false;
+
+	/** Whether distortion was evaluated this tick */
+	bool bWasDistortionEvaluated = false;
+
+	/** Whether or not nodal offset was applied to a tracked component this tick */
+	bool bWasNodalOffsetAppliedThisTick = false;
+
+	/** Latest focal length of the target camera, used to track external changes to the original (no overscan applied) focal length */
+	float LastFocalLength = -1.0f;
+
+	/** Delegate that is triggered when the LensModel changes */
+	FOnLensComponentModelChanged OnLensComponentModelChangedDelegate;
+};
diff --git a/Engine/Plugins/VirtualProduction/LensComponent/Source/LensComponentEditor/LensComponentEditor.build.cs b/Engine/Plugins/VirtualProduction/LensComponent/Source/LensComponentEditor/LensComponentEditor.build.cs
new file mode 100644
--- /dev/null
+++ b/Engine/Plugins/VirtualProduction/LensComponent/Source/LensComponentEditor/LensComponentEditor.build.cs
@@ -0,0 +1,30 @@
+// Copyright Epic Games, Inc. All Rights Reserved.
+
+namespace UnrealBuildTool.Rules
+{
+	public class LensComponentEditor : ModuleRules
+	{
+		public LensComponentEditor(ReadOnlyTargetRules Target) : base(Target)
+		{
+			PrivateDependencyModuleNames.AddRange(
+				new string[]
+				{
+					"CameraCalibrationCore",
+					"CameraCalibrationCoreEditor",
+					"Core",
+					"CoreUObject",
+					"Engine",
+					"LensComponent",
+					"LevelSequence",
+					"MovieScene",
+					"MovieSceneTools",
+					"Sequencer",
+					"Slate",
+					"SlateCore",
+					"TakeTrackRecorders",
+					"UnrealEd",
+				}
+			);
+		}
+	}
+}
diff --git a/Engine/Plugins/VirtualProduction/LensComponent/Source/LensComponentEditor/Private/LensComponentDetailCustomization.cpp b/Engine/Plugins/VirtualProduction/LensComponent/Source/LensComponentEditor/Private/LensComponentDetailCustomization.cpp
new file mode 100644
--- /dev/null
+++ b/Engine/Plugins/VirtualProduction/LensComponent/Source/LensComponentEditor/Private/LensComponentDetailCustomization.cpp
@@ -0,0 +1,59 @@
+// Copyright Epic Games, Inc. All Rights Reserved.
+
+#include "LensComponentDetailCustomization.h"
+
+#include "DetailLayoutBuilder.h"
+#include "LensDistortionStateDetailCustomization.h"
+
+FLensComponentDetailCustomization::~FLensComponentDetailCustomization()
+{
+	if (ULensComponent* LensComponent = WeakLensComponent.Get())
+	{
+		LensComponent->OnLensComponentModelChanged().RemoveAll(this);
+	}
+}
+
+void FLensComponentDetailCustomization::CustomizeDetails(IDetailLayoutBuilder& DetailBuilder)
+{
+	DetailLayout = &DetailBuilder;
+
+	TArray<TWeakObjectPtr<UObject>> SelectedObjects = DetailBuilder.GetSelectedObjects();
+
+ 	if (SelectedObjects.Num() != 1)
+ 	{
+ 		return;
+ 	}
+
+	if (ULensComponent* CustomizedLensComponent = Cast<ULensComponent>(SelectedObjects[0].Get()))
+	{
+		WeakLensComponent = CustomizedLensComponent;
+
+		if (!CustomizedLensComponent->OnLensComponentModelChanged().IsBoundToObject(this))
+		{
+			CustomizedLensComponent->OnLensComponentModelChanged().AddSP(this, &FLensComponentDetailCustomization::OnLensModelChanged);
+		}
+
+		TSharedRef<IPropertyHandle> LensModelPropertyHandle = DetailBuilder.GetProperty(GET_MEMBER_NAME_CHECKED(ULensComponent, LensModel));
+		LensModelPropertyHandle->SetOnPropertyValueChanged(FSimpleDelegate::CreateSP(this, &FLensComponentDetailCustomization::OnLensModelChanged));
+
+		const TSubclassOf<ULensModel> LensModel = CustomizedLensComponent->GetLensModel();
+
+		FOnGetPropertyTypeCustomizationInstance DistortionStateCustomizationDelegate = FOnGetPropertyTypeCustomizationInstance::CreateLambda(
+			[LensModel]
+			{
+				return MakeShared<FLensDistortionStateDetailCustomization>(LensModel);
+			});
+
+		DetailBuilder.RegisterInstancedCustomPropertyTypeLayout(FLensDistortionState::StaticStruct()->GetFName(), DistortionStateCustomizationDelegate);
+	}
+}
+
+void FLensComponentDetailCustomization::OnLensModelChanged()
+{
+	DetailLayout->ForceRefreshDetails();
+}
+
+void FLensComponentDetailCustomization::OnLensModelChanged(const TSubclassOf<ULensModel>& LensModel)
+{
+	DetailLayout->ForceRefreshDetails();
+}
diff --git a/Engine/Plugins/VirtualProduction/LensComponent/Source/LensComponentEditor/Private/LensComponentDetailCustomization.h b/Engine/Plugins/VirtualProduction/LensComponent/Source/LensComponentEditor/Private/LensComponentDetailCustomization.h
new file mode 100644
--- /dev/null
+++ b/Engine/Plugins/VirtualProduction/LensComponent/Source/LensComponentEditor/Private/LensComponentDetailCustomization.h
@@ -0,0 +1,40 @@
+// Copyright Epic Games, Inc. All Rights Reserved.
+
+#pragma once
+
+#include "IDetailCustomization.h"
+
+#include "LensComponent.h"
+
+/**
+* Customizes a ULensComponent details
+*/
+class FLensComponentDetailCustomization : public IDetailCustomization
+{
+public:
+
+	static TSharedRef<IDetailCustomization> MakeInstance()
+	{
+		return MakeShared<FLensComponentDetailCustomization>();
+	}
+
+	~FLensComponentDetailCustomization();
+
+	//~ Begin IDetailCustomization interface
+	virtual void CustomizeDetails(class IDetailLayoutBuilder& DetailBuilder) override;
+	//~ End IDetailCustomization interface
+
+protected:
+	/** Forces the details panel to refresh when the lens model changes */
+	void OnLensModelChanged();
+
+	/** Forces the details panel to refresh when the lens model changes */
+	void OnLensModelChanged(const TSubclassOf<ULensModel>& LensModel);
+
+protected:
+	/** Keep a reference to force refresh the layout */
+	IDetailLayoutBuilder* DetailLayout = nullptr;
+
+	/** Weak pointer to the lens component being edited */
+	TWeakObjectPtr<ULensComponent> WeakLensComponent;
+};
diff --git a/Engine/Plugins/VirtualProduction/LensComponent/Source/LensComponentEditor/Private/LensComponentEditorModule.cpp b/Engine/Plugins/VirtualProduction/LensComponent/Source/LensComponentEditor/Private/LensComponentEditorModule.cpp
new file mode 100644
--- /dev/null
+++ b/Engine/Plugins/VirtualProduction/LensComponent/Source/LensComponentEditor/Private/LensComponentEditorModule.cpp
@@ -0,0 +1,54 @@
+// Copyright Epic Games, Inc. All Rights Reserved.
+
+#include "Modules/ModuleInterface.h"
+
+#include "Editor.h"
+#include "Features/IModularFeatures.h"
+#include "ISequencerModule.h"
+#include "LensComponent.h"
+#include "LensComponentDetailCustomization.h"
+#include "Modules/ModuleManager.h"
+#include "MovieScene/LensComponentTrackEditor.h"
+#include "MovieScene/MovieSceneLensComponentTrackRecorder.h"
+#include "PropertyEditorModule.h"
+
+static const FName MovieSceneTrackRecorderFactoryName("MovieSceneTrackRecorderFactory");
+
+class FLensComponentEditorModule : public IModuleInterface
+{
+public:
+	//~ Begin IModuleInterface interface
+	virtual void StartupModule() override
+	{
+		FPropertyEditorModule& PropertyEditorModule = FModuleManager::Get().LoadModuleChecked<FPropertyEditorModule>("PropertyEditor");
+
+		PropertyEditorModule.RegisterCustomClassLayout(
+			ULensComponent::StaticClass()->GetFName(),
+			FOnGetDetailCustomizationInstance::CreateStatic(&FLensComponentDetailCustomization::MakeInstance)
+		);
+
+		// Register Track Editor for LensComponent track
+		ISequencerModule& SequencerModule = FModuleManager::LoadModuleChecked<ISequencerModule>("Sequencer");
+		CreateLensComponentTrackEditorHandle = SequencerModule.RegisterTrackEditor(FOnCreateTrackEditor::CreateStatic(&FLensComponentTrackEditor::CreateTrackEditor));
+
+		// Register modular feature for LensComponent track recorder factory
+		IModularFeatures::Get().RegisterModularFeature(MovieSceneTrackRecorderFactoryName, &LensComponentTrackRecorderFactory);
+	}
+
+	virtual void ShutdownModule() override
+	{
+		if (!IsEngineExitRequested() && GEditor && UObjectInitialized())
+		{
+			FPropertyEditorModule& PropertyModule = FModuleManager::GetModuleChecked<FPropertyEditorModule>("PropertyEditor");
+			PropertyModule.UnregisterCustomClassLayout(ULensComponent::StaticClass()->GetFName());
+		}
+	}
+	//~ End IModuleInterface interface
+
+private:
+	FDelegateHandle CreateLensComponentTrackEditorHandle;
+
+	FMovieSceneLensComponentTrackRecorderFactory LensComponentTrackRecorderFactory;
+};
+
+IMPLEMENT_MODULE(FLensComponentEditorModule, LensComponentEditor);
diff --git a/Engine/Plugins/VirtualProduction/LensComponent/Source/LensComponentEditor/Private/MovieScene/LensComponentTrackEditor.cpp b/Engine/Plugins/VirtualProduction/LensComponent/Source/LensComponentEditor/Private/MovieScene/LensComponentTrackEditor.cpp
new file mode 100644
--- /dev/null
+++ b/Engine/Plugins/VirtualProduction/LensComponent/Source/LensComponentEditor/Private/MovieScene/LensComponentTrackEditor.cpp
@@ -0,0 +1,25 @@
+// Copyright Epic Games, Inc. All Rights Reserved.
+
+#include "LensComponentTrackEditor.h"
+
+#include "LevelSequence.h"
+
+TSharedRef<ISequencerTrackEditor> FLensComponentTrackEditor::CreateTrackEditor(TSharedRef<ISequencer> InSequencer)
+{
+	return MakeShared<FLensComponentTrackEditor>(InSequencer);
+}
+
+bool FLensComponentTrackEditor::SupportsSequence(UMovieSceneSequence* InSequence) const
+{
+	if (!InSequence || InSequence->IsTrackSupported(UMovieSceneLensComponentTrack::StaticClass()) == ETrackSupport::NotSupported)
+	{
+		return false;
+	}
+
+	return InSequence->IsA(ULevelSequence::StaticClass());
+}
+
+bool FLensComponentTrackEditor::SupportsType(TSubclassOf<UMovieSceneTrack> Type) const
+{
+	return (Type == UMovieSceneLensComponentTrack::StaticClass());
+}
diff --git a/Engine/Plugins/VirtualProduction/LensComponent/Source/LensComponentEditor/Private/MovieScene/LensComponentTrackEditor.h b/Engine/Plugins/VirtualProduction/LensComponent/Source/LensComponentEditor/Private/MovieScene/LensComponentTrackEditor.h
new file mode 100644
--- /dev/null
+++ b/Engine/Plugins/VirtualProduction/LensComponent/Source/LensComponentEditor/Private/MovieScene/LensComponentTrackEditor.h
@@ -0,0 +1,30 @@
+// Copyright Epic Games, Inc. All Rights Reserved.
+
+#pragma once
+
+#include "KeyframeTrackEditor.h"
+#include "MovieScene/MovieSceneLensComponentTrack.h"
+
+/** A property track editor for Lens Component */
+class FLensComponentTrackEditor : public FKeyframeTrackEditor<UMovieSceneLensComponentTrack>
+{
+public:
+
+	FLensComponentTrackEditor(TSharedRef<ISequencer> InSequencer) : FKeyframeTrackEditor<UMovieSceneLensComponentTrack>(InSequencer)
+	{
+	}
+
+	/**
+	* Creates an instance of this class.  Called by a sequencer
+	*
+	* @param OwningSequencer The sequencer instance to be used by this tool
+	* @return The new instance of this class
+	*/
+	static TSharedRef<ISequencerTrackEditor> CreateTrackEditor(TSharedRef<ISequencer> OwningSequencer);
+
+public:
+	//~ Begin ISequencerTrackEditor interface
+	virtual bool SupportsSequence(UMovieSceneSequence* InSequence) const override;
+	virtual bool SupportsType(TSubclassOf<UMovieSceneTrack> Type) const override;
+	//~ End ISequencerTrackEditor interface
+};
diff --git a/Engine/Plugins/VirtualProduction/LensComponent/Source/LensComponentEditor/Private/MovieScene/MovieSceneLensComponentSection.cpp b/Engine/Plugins/VirtualProduction/LensComponent/Source/LensComponentEditor/Private/MovieScene/MovieSceneLensComponentSection.cpp
new file mode 100644
--- /dev/null
+++ b/Engine/Plugins/VirtualProduction/LensComponent/Source/LensComponentEditor/Private/MovieScene/MovieSceneLensComponentSection.cpp
@@ -0,0 +1,315 @@
+// Copyright Epic Games, Inc. All Rights Reserved.
+
+#include "MovieScene/MovieSceneLensComponentSection.h"
+
+#include "Channels/MovieSceneChannelProxy.h"
+#include "Evaluation/MovieScenePreAnimatedState.h"
+#include "Evaluation/PreAnimatedState/MovieScenePreAnimatedCaptureSource.h"
+#include "IMovieScenePlayer.h"
+#include "MovieSceneExecutionToken.h"
+#include "MovieSceneTimeHelpers.h"
+
+#define LOCTEXT_NAMESPACE "MovieSceneLensComponentSection"
+
+struct FPreAnimatedLensFileTokenProducer : IMovieScenePreAnimatedTokenProducer
+{
+	virtual IMovieScenePreAnimatedTokenPtr CacheExistingState(UObject& Object) const
+	{
+		struct FLensFileToken : IMovieScenePreAnimatedToken
+		{
+			FLensFileToken(ULensComponent* LensComponent)
+			{
+				// Save the original values of these settings
+				FLensFilePicker LensFilePicker = LensComponent->GetLensFilePicker();
+				bOriginalUseDefaultLensFile = LensFilePicker.bUseDefaultLensFile;
+				OriginalLensFile = LensFilePicker.GetLensFile();
+				OriginalEvalMode = LensComponent->GetFIZEvaluationMode();
+			}
+
+			virtual void RestoreState(UObject& Object, const UE::MovieScene::FRestoreStateParams& Params)
+			{
+				if (ULensComponent* LensComponent = CastChecked<ULensComponent>(&Object))
+				{
+					// Restore the original values of these settings
+					LensComponent->SetLensFilePicker({ bOriginalUseDefaultLensFile, OriginalLensFile.Get() });
+					LensComponent->SetFIZEvaluationMode(OriginalEvalMode);
+				}
+			}
+
+			bool bOriginalUseDefaultLensFile = false;
+			TWeakObjectPtr<ULensFile> OriginalLensFile;
+			EFIZEvaluationMode OriginalEvalMode = EFIZEvaluationMode::UseRecordedValues;
+		};
+
+		return FLensFileToken(CastChecked<ULensComponent>(&Object));
+	}
+	static FMovieSceneAnimTypeID GetAnimTypeID()
+	{
+		return TMovieSceneAnimTypeID<FPreAnimatedLensFileTokenProducer>();
+	}
+};
+
+void UMovieSceneLensComponentSection::Initialize(ULensComponent* Component)
+{
+	// Cache a pointer to the recorded component to retrieve its data later during RecordFrame()
+	RecordedComponent = Component;
+
+	// Cache the Lens Model used by the recorded lens component
+	LensModelClass = RecordedComponent->GetLensModel();
+
+	// Duplicate the recorded LensFile asset and save the duplicate in this section
+ 	if (ULensFile* LensFile = Component->GetLensFile())
+ 	{
+ 		const FName DuplicateName = FName(*FString::Format(TEXT("{0}_Cached"), { LensFile->GetFName().ToString() }));
+ 		CachedLensFile = Cast<ULensFile>(StaticDuplicateObject(LensFile, this, DuplicateName, RF_AllFlags & ~RF_Transient));
+ 	}
+
+	// Create the channels that will be used to record the distortion state
+	CreateChannelProxy();
+}
+
+void UMovieSceneLensComponentSection::PostEditImport()
+{
+	Super::PostEditImport();
+
+	UpdateChannelProxy();
+}
+
+void UMovieSceneLensComponentSection::PostLoad()
+{
+	Super::PostLoad();
+
+	UpdateChannelProxy();
+}
+
+#if WITH_EDITOR
+void UMovieSceneLensComponentSection::AddChannelWithEditor(FMovieSceneChannelProxyData& ChannelProxyData, FMovieSceneFloatChannel& Channel, FText GroupName, FText ChannelName, int32 SortOrder)
+{
+	FMovieSceneChannelMetaData ChannelEditorData(FName(*(ChannelName.ToString())), ChannelName, GroupName);
+	ChannelEditorData.SortOrder = SortOrder;
+	ChannelEditorData.bCanCollapseToTrack = false;
+
+	ChannelProxyData.Add(Channel, ChannelEditorData, TMovieSceneExternalValue<float>());
+}
+#endif // WITH_EDITOR
+
+void UMovieSceneLensComponentSection::CreateChannelProxy()
+{
+	// Add channels for distortion parameters
+	if (LensModelClass)
+	{
+		if (const ULensModel* const LensModel = LensModelClass->GetDefaultObject<ULensModel>())
+		{
+			const int32 ParameterCount = LensModel->GetNumParameters();
+			DistortionParameterChannels.SetNum(ParameterCount);
+		}
+	}
+
+	// Add channels for FxFy
+	FxFyChannels.SetNum(2);
+
+	// Add channels for Image Center
+	ImageCenterChannels.SetNum(2);
+
+	UpdateChannelProxy();
+}
+
+void UMovieSceneLensComponentSection::UpdateChannelProxy()
+{
+	FMovieSceneChannelProxyData Channels;
+
+	// Add channels for distortion parameters
+	if (LensModelClass)
+	{
+		if (const ULensModel* const LensModel = LensModelClass->GetDefaultObject<ULensModel>())
+		{
+			const int32 ParameterCount = LensModel->GetNumParameters();
+			if (ParameterCount == DistortionParameterChannels.Num())
+			{
+#if WITH_EDITOR
+				TArray<FText> ParameterDisplayNames = LensModel->GetParameterDisplayNames();
+				const FText ParameterGroupName = LOCTEXT("ParameterGroupName", "Distortion Parameters");
+#endif // WITH_EDITOR
+
+				for (int32 Index = 0; Index < ParameterCount; ++Index)
+				{
+#if WITH_EDITOR
+					AddChannelWithEditor(Channels, DistortionParameterChannels[Index], ParameterGroupName, ParameterDisplayNames[Index], Index);
+#else
+					Channels.Add(DistortionParameterChannels[Index]);
+#endif // WITH_EDITOR
+				}
+			}
+		}
+	}
+
+	// Add channels for FxFy
+	if (FxFyChannels.Num() == 2)
+	{
+#if WITH_EDITOR
+		const FText FxFyGroupName = LOCTEXT("FxFyGroupName", "FxFy");
+		constexpr int32 FxFySortOrder = 1000;
+		AddChannelWithEditor(Channels, FxFyChannels[0], FxFyGroupName, LOCTEXT("FxChannelName", "Fx"), FxFySortOrder + 0);
+		AddChannelWithEditor(Channels, FxFyChannels[1], FxFyGroupName, LOCTEXT("FyChannelName", "Fy"), FxFySortOrder + 1);
+#else
+		Channels.Add(FxFyChannels[0]);
+		Channels.Add(FxFyChannels[1]);
+#endif // WITH_EDITOR
+	}
+
+	// Add channels for Image Center
+	if (ImageCenterChannels.Num() == 2)
+	{
+#if WITH_EDITOR
+		const FText ImageCenterGroupName = LOCTEXT("CxCyGroupName", "Image Center");
+		constexpr int32 ImageCenterSortOrder = 2000;
+		AddChannelWithEditor(Channels, ImageCenterChannels[0], ImageCenterGroupName, LOCTEXT("CxChannelName", "Cx"), ImageCenterSortOrder + 0);
+		AddChannelWithEditor(Channels, ImageCenterChannels[1], ImageCenterGroupName, LOCTEXT("CyChannelName", "Cy"), ImageCenterSortOrder + 1);
+#else
+		Channels.Add(ImageCenterChannels[0]);
+		Channels.Add(ImageCenterChannels[1]);
+#endif // WITH_EDITOR
+	}
+
+	ChannelProxy = MakeShared<FMovieSceneChannelProxy>(MoveTemp(Channels));
+}
+
+void UMovieSceneLensComponentSection::RecordFrame(FFrameNumber FrameNumber)
+{
+	if (const ULensComponent* const LensComponent = RecordedComponent.Get())
+	{
+		// Get the current distortion state of the lens component being recorded
+		const FLensDistortionState CurrentState = LensComponent->GetDistortionState();
+		const int32 ParameterCount = CurrentState.DistortionInfo.Parameters.Num();
+
+		// Record the current state of each distortion parameter
+		for (int32 Index = 0; Index < ParameterCount; ++Index)
+		{
+			DistortionParameterChannels[Index].GetData().AddKey(FrameNumber, FMovieSceneFloatValue(CurrentState.DistortionInfo.Parameters[Index]));
+		}
+
+		// Record the current state of FxFy
+		FxFyChannels[0].GetData().AddKey(FrameNumber, FMovieSceneFloatValue(CurrentState.FocalLengthInfo.FxFy[0]));
+		FxFyChannels[1].GetData().AddKey(FrameNumber, FMovieSceneFloatValue(CurrentState.FocalLengthInfo.FxFy[1]));
+
+		// Record the current state of the image center
+		ImageCenterChannels[0].GetData().AddKey(FrameNumber, FMovieSceneFloatValue(CurrentState.ImageCenter.PrincipalPoint[0]));
+		ImageCenterChannels[1].GetData().AddKey(FrameNumber, FMovieSceneFloatValue(CurrentState.ImageCenter.PrincipalPoint[1]));
+	}
+}
+
+void UMovieSceneLensComponentSection::Finalize()
+{
+	// Reduce keys on each of the channels that we recorded
+	FKeyDataOptimizationParams Params;
+
+	const int32 ParameterCount = DistortionParameterChannels.Num();
+	for (int32 Index = 0; Index < ParameterCount; ++Index)
+	{
+		UE::MovieScene::Optimize(&DistortionParameterChannels[Index], Params);
+	}
+
+	UE::MovieScene::Optimize(&FxFyChannels[0], Params);
+	UE::MovieScene::Optimize(&FxFyChannels[1], Params);
+	UE::MovieScene::Optimize(&ImageCenterChannels[0], Params);
+	UE::MovieScene::Optimize(&ImageCenterChannels[1], Params);
+}
+
+void UMovieSceneLensComponentSection::Begin(IMovieScenePlayer* Player, const UE::MovieScene::FEvaluationHookParams& Params) const
+{
+	FScopedPreAnimatedCaptureSource CaptureSource(&Player->PreAnimatedState, this, Params.SequenceID, EvalOptions.CompletionMode == EMovieSceneCompletionMode::RestoreState);
+
+	for (TWeakObjectPtr<> BoundObject : Player->FindBoundObjects(Params.ObjectBindingID, Params.SequenceID))
+	{
+		if (ULensComponent* LensComponent = Cast<ULensComponent>(BoundObject.Get()))
+		{
+			Player->SavePreAnimatedState(*LensComponent, FPreAnimatedLensFileTokenProducer::GetAnimTypeID(), FPreAnimatedLensFileTokenProducer());
+
+			// Override the existing lens file settings of the lens component to use the cached, duplicate lens file
+			LensComponent->SetFIZEvaluationMode(EFIZEvaluationMode::UseRecordedValues);
+
+			// TODO: Check that the number of parameters for the lens component's current lens model matches the number of recorded parameter channels. 
+		}
+	}
+}
+
+void UMovieSceneLensComponentSection::Update(IMovieScenePlayer* Player, const UE::MovieScene::FEvaluationHookParams& Params) const
+{
+	for (TWeakObjectPtr<> BoundObject : Player->FindBoundObjects(Params.ObjectBindingID, Params.SequenceID))
+	{
+		if (ULensComponent* LensComponent = Cast<ULensComponent>(BoundObject.Get()))
+		{
+			// Set the LensFile to use on the LensComponent, either the cached version or an override
+			if (OverrideLensFile)
+			{
+				LensComponent->SetLensFilePicker({ false, OverrideLensFile });
+			}
+			else
+			{
+				LensComponent->SetLensFilePicker({ false, CachedLensFile });
+			}
+
+			// Nodal offset is baked in to the transform of the camera, but users can use this flag to force the LensComponent to re-evaluate
+			// the LensFile and re-apply a new nodal offset instead.
+			if (bReapplyNodalOffset)
+			{
+				LensComponent->ReapplyNodalOffset();
+			}
+
+			// Get the end frame of the current range being evaluated
+			const TRange<FFrameNumber> PlaybackRange = Params.Context.GetFrameNumberRange();
+			const FFrameNumber ExclusiveEndFrame = UE::MovieScene::DiscreteExclusiveUpper(PlaybackRange);
+
+			// Get the current distortion state of the lens component, which will be overwritten with values from the channel (if present)
+			FLensDistortionState NewState = LensComponent->GetDistortionState();
+
+			// Update the state from the distortion parameter channels
+			const int32 ParameterCount = DistortionParameterChannels.Num();
+			NewState.DistortionInfo.Parameters.SetNum(ParameterCount);
+
+			for (int32 Index = 0; Index < ParameterCount; ++Index)
+			{
+				float NewValue = 0.0f;
+				if (DistortionParameterChannels[Index].Evaluate(ExclusiveEndFrame, NewValue))
+				{
+					NewState.DistortionInfo.Parameters[Index] = NewValue;
+				}
+			}
+
+			// Update the state from the FxFy channels
+			float NewValueFx = 0.0f;
+			if (FxFyChannels[0].Evaluate(ExclusiveEndFrame, NewValueFx))
+			{
+				NewState.FocalLengthInfo.FxFy[0] = NewValueFx;
+			}
+
+			float NewValueFy = 0.0f;
+			if (FxFyChannels[1].Evaluate(ExclusiveEndFrame, NewValueFy))
+			{
+				NewState.FocalLengthInfo.FxFy[1] = NewValueFy;
+			}
+
+			// Update the state from the ImageCenter channels
+			float NewValueCx = 0.0f;
+			if (ImageCenterChannels[0].Evaluate(ExclusiveEndFrame, NewValueCx))
+			{
+				NewState.ImageCenter.PrincipalPoint[0] = NewValueCx;
+			}
+
+			float NewValueCy = 0.0f;
+			if (ImageCenterChannels[1].Evaluate(ExclusiveEndFrame, NewValueCy))
+			{
+				NewState.ImageCenter.PrincipalPoint[1] = NewValueCy;
+			}
+
+			LensComponent->SetDistortionState(NewState);
+		}
+	}
+}
+
+void UMovieSceneLensComponentSection::End(IMovieScenePlayer* Player, const UE::MovieScene::FEvaluationHookParams& Params) const
+{
+	Player->RestorePreAnimatedState();
+}
+
+#undef LOCTEXT_NAMESPACE
diff --git a/Engine/Plugins/VirtualProduction/LensComponent/Source/LensComponentEditor/Private/MovieScene/MovieSceneLensComponentTrack.cpp b/Engine/Plugins/VirtualProduction/LensComponent/Source/LensComponentEditor/Private/MovieScene/MovieSceneLensComponentTrack.cpp
new file mode 100644
--- /dev/null
+++ b/Engine/Plugins/VirtualProduction/LensComponent/Source/LensComponentEditor/Private/MovieScene/MovieSceneLensComponentTrack.cpp
@@ -0,0 +1,64 @@
+// Copyright Epic Games, Inc. All Rights Reserved.
+
+#include "MovieScene/MovieSceneLensComponentTrack.h"
+
+#include "MovieScene/MovieSceneLensComponentSection.h"
+
+#define LOCTEXT_NAMESPACE "MovieSceneLensComponentTrack"
+
+bool UMovieSceneLensComponentTrack::SupportsType(TSubclassOf<UMovieSceneSection> SectionClass) const
+{
+	return SectionClass == UMovieSceneLensComponentSection::StaticClass();
+}
+
+UMovieSceneSection* UMovieSceneLensComponentTrack::CreateNewSection()
+{
+	return NewObject<UMovieSceneSection>(this, UMovieSceneLensComponentSection::StaticClass(), MakeUniqueObjectName(this, UMovieSceneLensComponentSection::StaticClass(), TEXT("LensComponentSection")), RF_Transactional);
+}
+
+#if WITH_EDITORONLY_DATA
+FText UMovieSceneLensComponentTrack::GetDisplayName() const
+{
+	return LOCTEXT("LensComponentTrackDisplayName", "Lens Component");
+}
+#endif
+
+void UMovieSceneLensComponentTrack::AddSection(UMovieSceneSection& Section)
+{
+	Sections.Add(&Section);
+}
+
+const TArray<UMovieSceneSection*>& UMovieSceneLensComponentTrack::GetAllSections() const
+{
+	return Sections;
+}
+
+bool UMovieSceneLensComponentTrack::HasSection(const UMovieSceneSection& Section) const
+{
+	return Sections.Contains(&Section);
+}
+
+bool UMovieSceneLensComponentTrack::IsEmpty() const
+{
+	return Sections.IsEmpty();
+}
+
+void UMovieSceneLensComponentTrack::RemoveAllAnimationData()
+{
+	Sections.Empty();
+}
+
+void UMovieSceneLensComponentTrack::RemoveSection(UMovieSceneSection& Section)
+{
+	Sections.Remove(&Section);
+}
+
+void UMovieSceneLensComponentTrack::RemoveSectionAt(int32 SectionIndex)
+{
+	if (Sections.IsValidIndex(SectionIndex))
+	{
+		Sections.RemoveAt(SectionIndex);
+	}
+}
+
+#undef LOCTEXT_NAMESPACE
diff --git a/Engine/Plugins/VirtualProduction/LensComponent/Source/LensComponentEditor/Private/MovieScene/MovieSceneLensComponentTrackRecorder.cpp b/Engine/Plugins/VirtualProduction/LensComponent/Source/LensComponentEditor/Private/MovieScene/MovieSceneLensComponentTrackRecorder.cpp
new file mode 100644
--- /dev/null
+++ b/Engine/Plugins/VirtualProduction/LensComponent/Source/LensComponentEditor/Private/MovieScene/MovieSceneLensComponentTrackRecorder.cpp
@@ -0,0 +1,82 @@
+// Copyright Epic Games, Inc. All Rights Reserved.
+
+#include "MovieSceneLensComponentTrackRecorder.h"
+
+#include "MovieScene/MovieSceneLensComponentTrack.h"
+#include "LensComponent.h"
+
+#define LOCTEXT_NAMESPACE "MovieSceneLensComponentTrackRecorder"
+
+bool FMovieSceneLensComponentTrackRecorderFactory::CanRecordObject(class UObject* InObjectToRecord) const
+{
+	return (InObjectToRecord->GetClass() == ULensComponent::StaticClass());
+}
+
+bool FMovieSceneLensComponentTrackRecorderFactory::CanRecordProperty(class UObject* InObjectToRecord, class FProperty* InPropertyToRecord) const
+{ 
+	return false; 
+}
+
+UMovieSceneTrackRecorder* FMovieSceneLensComponentTrackRecorderFactory::CreateTrackRecorderForObject() const
+{
+	return NewObject<UMovieSceneLensComponentTrackRecorder>();
+}
+
+UMovieSceneTrackRecorder* FMovieSceneLensComponentTrackRecorderFactory::CreateTrackRecorderForProperty(UObject* InObjectToRecord, const FName& InPropertyToRecord) const
+{
+	return nullptr; 
+}
+
+FText FMovieSceneLensComponentTrackRecorderFactory::GetDisplayName() const
+{ 
+	return LOCTEXT("LensComponentTrackRecorderDisplayName", "Lens Component"); 
+}
+
+void UMovieSceneLensComponentTrackRecorder::CreateTrackImpl()
+{
+	// Find/Add a Lens Component Track to the MovieScene
+	UMovieSceneLensComponentTrack* LensComponentTrack = MovieScene->FindTrack<UMovieSceneLensComponentTrack>(ObjectGuid);
+	if (!LensComponentTrack)
+	{
+		LensComponentTrack = MovieScene->AddTrack<UMovieSceneLensComponentTrack>(ObjectGuid);
+	}
+	else
+	{
+		LensComponentTrack->RemoveAllAnimationData();
+	}
+
+	// Create a new section for the Lens Component Track and initialize it with a pointer to the component being recorded
+	if (LensComponentTrack)
+	{
+		LensComponentSection = Cast<UMovieSceneLensComponentSection>(LensComponentTrack->CreateNewSection());
+
+		if (UMovieSceneLensComponentSection* Section = LensComponentSection.Get())
+		{
+			LensComponentTrack->AddSection(*Section);
+
+			ULensComponent* LensComponent = Cast<ULensComponent>(ObjectToRecord.Get());
+			Section->Initialize(LensComponent);
+		}
+	}
+}
+
+void UMovieSceneLensComponentTrackRecorder::RecordSampleImpl(const FQualifiedFrameTime& CurrentTime)
+{
+	if (UMovieSceneLensComponentSection* Section = LensComponentSection.Get())
+	{
+		FFrameRate TickResolution = Section->GetTypedOuter<UMovieScene>()->GetTickResolution();
+		FFrameNumber CurrentFrame = CurrentTime.ConvertTo(TickResolution).FloorToFrame();
+		Section->RecordFrame(CurrentFrame);
+		Section->SetEndFrame(CurrentFrame);
+	}
+}
+
+void UMovieSceneLensComponentTrackRecorder::FinalizeTrackImpl()
+{
+	if (UMovieSceneLensComponentSection* Section = LensComponentSection.Get())
+	{
+		Section->Finalize();
+	}
+}
+
+#undef LOCTEXT_NAMESPACE
diff --git a/Engine/Plugins/VirtualProduction/LensComponent/Source/LensComponentEditor/Private/MovieScene/MovieSceneLensComponentTrackRecorder.h b/Engine/Plugins/VirtualProduction/LensComponent/Source/LensComponentEditor/Private/MovieScene/MovieSceneLensComponentTrackRecorder.h
new file mode 100644
--- /dev/null
+++ b/Engine/Plugins/VirtualProduction/LensComponent/Source/LensComponentEditor/Private/MovieScene/MovieSceneLensComponentTrackRecorder.h
@@ -0,0 +1,43 @@
+// Copyright Epic Games, Inc. All Rights Reserved.
+
+#pragma once
+
+#include "TrackRecorders/IMovieSceneTrackRecorderFactory.h"
+#include "TrackRecorders/MovieSceneTrackRecorder.h"
+
+#include "MovieScene/MovieSceneLensComponentSection.h"
+
+#include "MovieSceneLensComponentTrackRecorder.generated.h"
+
+/** MovieScene Track Recorder Factory for a Lens Component */
+class FMovieSceneLensComponentTrackRecorderFactory : public IMovieSceneTrackRecorderFactory
+{
+public:
+	//~ Begin IMovieSceneTrackRecorderFactory interface
+	virtual bool CanRecordObject(class UObject* InObjectToRecord) const override;
+	virtual bool CanRecordProperty(class UObject* InObjectToRecord, class FProperty* InPropertyToRecord) const override;
+
+	virtual UMovieSceneTrackRecorder* CreateTrackRecorderForObject() const override;
+	virtual UMovieSceneTrackRecorder* CreateTrackRecorderForProperty(UObject* InObjectToRecord, const FName& InPropertyToRecord) const override;
+
+	virtual FText GetDisplayName() const override;
+	//~ End IMovieSceneTrackRecorderFactory interface
+};
+
+/** MovieScene Track Recorder for a Lens Component */
+UCLASS(BlueprintType)
+class UMovieSceneLensComponentTrackRecorder : public UMovieSceneTrackRecorder
+{
+	GENERATED_BODY()
+public:
+	//~ Begin UMovieSceneTrackRecorder Interface
+	virtual void CreateTrackImpl() override;
+	virtual void RecordSampleImpl(const FQualifiedFrameTime& CurrentTime) override;
+	virtual void FinalizeTrackImpl() override;
+	virtual UMovieSceneSection* GetMovieSceneSection() const override { return LensComponentSection.Get(); }
+	//~ End UMovieSceneTrackRecorder Interface
+
+private:
+	/** Section in which to record Lens Component data */
+	TWeakObjectPtr<UMovieSceneLensComponentSection> LensComponentSection;
+};
diff --git a/Engine/Plugins/VirtualProduction/LensComponent/Source/LensComponentEditor/Public/MovieScene/MovieSceneLensComponentSection.h b/Engine/Plugins/VirtualProduction/LensComponent/Source/LensComponentEditor/Public/MovieScene/MovieSceneLensComponentSection.h
new file mode 100644
--- /dev/null
+++ b/Engine/Plugins/VirtualProduction/LensComponent/Source/LensComponentEditor/Public/MovieScene/MovieSceneLensComponentSection.h
@@ -0,0 +1,92 @@
+// Copyright Epic Games, Inc. All Rights Reserved.
+
+#pragma once
+
+#include "Sections/MovieSceneHookSection.h"
+
+#include "Channels/MovieSceneFloatChannel.h"
+#include "LensComponent.h"
+#include "LensFile.h"
+
+#include "MovieSceneLensComponentSection.generated.h"
+
+struct FMovieSceneChannelProxyData;
+
+/** Movie Scene section for Lens Component */
+UCLASS()
+class LENSCOMPONENTEDITOR_API UMovieSceneLensComponentSection : public UMovieSceneHookSection
+{
+	GENERATED_BODY()
+
+public:
+
+	UMovieSceneLensComponentSection(const FObjectInitializer& ObjectInitializer)
+		: UMovieSceneHookSection(ObjectInitializer) 
+	{ 
+		bRequiresRangedHook = true; 
+	}
+
+	//~ Begin UMovieSceneHookSection interface
+	virtual void Begin(IMovieScenePlayer* Player, const UE::MovieScene::FEvaluationHookParams& Params) const override;
+	virtual void Update(IMovieScenePlayer* Player, const UE::MovieScene::FEvaluationHookParams& Params) const override;
+	virtual void End(IMovieScenePlayer* Player, const UE::MovieScene::FEvaluationHookParams& Params) const override;
+	//~ End UMovieSceneHookSection interface
+
+	/** Initialize the section with the Lens Component whose properties it will record */
+	void Initialize(ULensComponent* Component);
+
+	/** Reduce keys in each of the recorded channels */
+	void Finalize();
+
+	/** Record the current value of each distortion state property to its appropriate track */
+	void RecordFrame(FFrameNumber FrameNumber);
+
+	//~ Begin UObject interface
+	virtual void PostEditImport() override;
+	virtual void PostLoad() override;
+	//~ End UObject interface
+
+private:
+	/** Create the channels in which distortion state values will be recorded */
+	void CreateChannelProxy();
+
+	/** Add the distortion state channels to the channel proxy */
+	void UpdateChannelProxy();
+
+#if WITH_EDITOR
+	void AddChannelWithEditor(FMovieSceneChannelProxyData& ChannelProxyData, FMovieSceneFloatChannel& Channel, FText GroupName, FText ChannelName, int32 SortOrder);
+#endif // WITH_EDITOR
+
+public:
+	/** If true, then every Update, the nodal offset will be re-evaluated on the lens component */
+	UPROPERTY(EditAnywhere, Category = "Camera Calibration")
+	bool bReapplyNodalOffset = false;
+
+	/** LensFile asset that should be used instead of the cached LensFile during playback */
+	UPROPERTY(EditAnywhere, Category = "Camera Calibration")
+	TObjectPtr<ULensFile> OverrideLensFile = nullptr;
+
+private:
+	/** Saved duplicate of the LensFile asset used by the recorded Lens Component at the time of recording */
+	UPROPERTY(VisibleAnywhere, Category="Camera Calibration")
+	TObjectPtr<ULensFile> CachedLensFile = nullptr;
+
+	/** Channels to store Distortion Parameter values (will be sized during initialization based on the LensComponent's LensModel) */
+	UPROPERTY()
+	TArray<FMovieSceneFloatChannel> DistortionParameterChannels;
+
+	/** Channels to store FxFy values */
+	UPROPERTY()
+	TArray<FMovieSceneFloatChannel> FxFyChannels;
+
+	/** Channels to store Image Center values */
+	UPROPERTY()
+	TArray<FMovieSceneFloatChannel> ImageCenterChannels;
+
+	/** The Lens Model used by the recorded LensComponent */
+	UPROPERTY()
+	TSubclassOf<ULensModel> LensModelClass;
+
+private:
+	TWeakObjectPtr<ULensComponent> RecordedComponent;
+};
diff --git a/Engine/Plugins/VirtualProduction/LensComponent/Source/LensComponentEditor/Public/MovieScene/MovieSceneLensComponentTrack.h b/Engine/Plugins/VirtualProduction/LensComponent/Source/LensComponentEditor/Public/MovieScene/MovieSceneLensComponentTrack.h
new file mode 100644
--- /dev/null
+++ b/Engine/Plugins/VirtualProduction/LensComponent/Source/LensComponentEditor/Public/MovieScene/MovieSceneLensComponentTrack.h
@@ -0,0 +1,37 @@
+// Copyright Epic Games, Inc. All Rights Reserved.
+
+#pragma once
+
+#include "MovieSceneNameableTrack.h"
+
+#include "MovieSceneLensComponentTrack.generated.h"
+
+/** Movie Scene track for Lens Component */
+UCLASS()
+class LENSCOMPONENTEDITOR_API UMovieSceneLensComponentTrack : public UMovieSceneNameableTrack
+{
+	GENERATED_BODY()
+
+public:
+	//~ Begin UMovieSceneTrack interface
+	virtual void AddSection(UMovieSceneSection& Section) override;
+	virtual const TArray<UMovieSceneSection*>& GetAllSections() const override;
+	virtual bool HasSection(const UMovieSceneSection& Section) const override;
+	virtual bool IsEmpty() const override;
+	virtual void RemoveAllAnimationData() override;
+	virtual void RemoveSection(UMovieSceneSection& Section) override;
+	virtual void RemoveSectionAt(int32 SectionIndex) override;
+
+	virtual bool SupportsType(TSubclassOf<UMovieSceneSection> SectionClass) const override;
+	virtual UMovieSceneSection* CreateNewSection() override;
+
+#if WITH_EDITORONLY_DATA
+	virtual FText GetDisplayName() const override;
+#endif
+	//~ End UMovieSceneTrack interface
+
+protected:
+	/** Array of movie scene sections managed by this track */
+	UPROPERTY()
+	TArray<TObjectPtr<UMovieSceneSection>> Sections;
+};
\ No newline at end of file
diff --git a/Engine/Plugins/VirtualProduction/LiveLinkCamera/LiveLinkCamera.uplugin b/Engine/Plugins/VirtualProduction/LiveLinkCamera/LiveLinkCamera.uplugin
--- a/Engine/Plugins/VirtualProduction/LiveLinkCamera/LiveLinkCamera.uplugin
+++ b/Engine/Plugins/VirtualProduction/LiveLinkCamera/LiveLinkCamera.uplugin
@@ -23,6 +23,10 @@
 		{
 			"Name": "CameraCalibrationCore",
 			"Enabled": true
+		},
+		{
+			"Name": "LensComponent",
+			"Enabled": true
 		}
 	],
 	"Modules": [
diff --git a/Engine/Plugins/VirtualProduction/LiveLinkCamera/Source/LiveLinkCamera/LiveLinkCamera.build.cs b/Engine/Plugins/VirtualProduction/LiveLinkCamera/Source/LiveLinkCamera/LiveLinkCamera.build.cs
--- a/Engine/Plugins/VirtualProduction/LiveLinkCamera/Source/LiveLinkCamera/LiveLinkCamera.build.cs
+++ b/Engine/Plugins/VirtualProduction/LiveLinkCamera/Source/LiveLinkCamera/LiveLinkCamera.build.cs
@@ -19,5 +19,12 @@ public class LiveLinkCamera : ModuleRules
 				"LiveLinkComponents",
 			}
 		);
+
+		PrivateDependencyModuleNames.AddRange(
+			new string[]
+			{
+				"LensComponent",
+			}
+		);
 	}
 }
diff --git a/Engine/Plugins/VirtualProduction/LiveLinkLens/LiveLinkLens.uplugin b/Engine/Plugins/VirtualProduction/LiveLinkLens/LiveLinkLens.uplugin
--- a/Engine/Plugins/VirtualProduction/LiveLinkLens/LiveLinkLens.uplugin
+++ b/Engine/Plugins/VirtualProduction/LiveLinkLens/LiveLinkLens.uplugin
@@ -30,6 +30,10 @@
 		{
 			"Name": "CameraCalibrationCore",
 			"Enabled": true
+		},
+		{
+			"Name": "LensComponent",
+			"Enabled": true
 		}
 	]
 }
\ No newline at end of file
diff --git a/Engine/Plugins/VirtualProduction/LiveLinkLens/Source/LiveLinkLens/LiveLinkLens.Build.cs b/Engine/Plugins/VirtualProduction/LiveLinkLens/Source/LiveLinkLens/LiveLinkLens.Build.cs
--- a/Engine/Plugins/VirtualProduction/LiveLinkLens/Source/LiveLinkLens/LiveLinkLens.Build.cs
+++ b/Engine/Plugins/VirtualProduction/LiveLinkLens/Source/LiveLinkLens/LiveLinkLens.Build.cs
@@ -25,6 +25,7 @@ public class LiveLinkLens : ModuleRules
 				"Core",
 				"CoreUObject",
 				"Engine",
+				"LensComponent",
 				"LiveLinkMovieScene",
 				"MovieScene",
 				"MovieSceneTracks",
