LensComponent: Expose GetLensDistortionHandler to Blueprint.

Replaces the functionality that used to be available via the Camera Calibration Subsystem.

diff --git a/Engine/Plugins/VirtualProduction/LensComponent/Source/LensComponent/Public/LensComponent.h b/Engine/Plugins/VirtualProduction/LensComponent/Source/LensComponent/Public/LensComponent.h
--- a/Engine/Plugins/VirtualProduction/LensComponent/Source/LensComponent/Public/LensComponent.h
+++ b/Engine/Plugins/VirtualProduction/LensComponent/Source/LensComponent/Public/LensComponent.h
@@ -207,6 +207,7 @@ public:
 
 public:
 	/** Returns the LensDistortionHandler in use for the current LensModel */
+	UFUNCTION(BlueprintCallable, Category = "Lens Component")
 	ULensDistortionModelHandlerBase* GetLensDistortionHandler() const;
 
 	/** Reset the tracked component back to its original tracked pose and reapply nodal offset to it by re-evaluating the LensFile */
