CameraCalibration: Fix a bug where distortion is applied to a camera even when the "Apply Distortion" setting is false on a Lens Component.

The previous logic that removed the distortion post-process material only did so in PostEditChange, and so other clients connected in a multi-user session would change the setting to false, but not trigger a PostEditChange.

#rnx

diff --git a/Engine/Plugins/VirtualProduction/LensComponent/Source/LensComponent/Private/LensComponent.cpp b/Engine/Plugins/VirtualProduction/LensComponent/Source/LensComponent/Private/LensComponent.cpp
--- a/Engine/Plugins/VirtualProduction/LensComponent/Source/LensComponent/Private/LensComponent.cpp
+++ b/Engine/Plugins/VirtualProduction/LensComponent/Source/LensComponent/Private/LensComponent.cpp
@@ -214,6 +214,8 @@ void ULensComponent::TickComponent(float DeltaTime, ELevelTick TickType, FActorC
 			{
 				CineCameraComponent->CurrentFocalLength = OriginalFocalLength;
 			}
+
+			CleanupDistortion(CineCameraComponent);
 		}
 
 		LastFocalLength = CineCameraComponent->CurrentFocalLength;
