LensComponent: Transact camera changes when cleaning up distortion.

When apply distortion is cleared on the lens component, clear distortion is also called which clears the property that indicates that distortion is set up, which is something that clear distortion looks at before doing anything. Therefore the MU client that receives this transaction never actually cleans up the distortion.

This change adds the camera component to the ongoing transaction when it distortion is cleaned up, such that the receiver gets the value as well.

diff --git a/Engine/Plugins/VirtualProduction/LensComponent/Source/LensComponent/Private/LensComponent.cpp b/Engine/Plugins/VirtualProduction/LensComponent/Source/LensComponent/Private/LensComponent.cpp
--- a/Engine/Plugins/VirtualProduction/LensComponent/Source/LensComponent/Private/LensComponent.cpp
+++ b/Engine/Plugins/VirtualProduction/LensComponent/Source/LensComponent/Private/LensComponent.cpp
@@ -723,8 +723,13 @@ float ULensComponent::GetDesqueezedSensorWidth(UCineCameraComponent* const CineC
 
 void ULensComponent::CleanupDistortion(UCineCameraComponent* const CineCameraComponent)
 {
-	if (bIsDistortionSetup)
+	if (bIsDistortionSetup && ensure(CineCameraComponent))
 	{
+#if WITH_EDITOR
+		// If we are in a transaction, we also want to capture the changes made to the camera.
+		CineCameraComponent->Modify();
+#endif // WITH_EDITOR
+
 		// Remove the last distortion MID that was applied to the target camera component
 		if (LastDistortionMID)
 		{
