CameraCalibration: Add support to the lens component to drive the overscan properties of the camera component

diff --git a/Engine/Plugins/VirtualProduction/CameraCalibrationCore/Source/CameraCalibrationCore/Private/SphericalLensDistortionModelHandler.cpp b/Engine/Plugins/VirtualProduction/CameraCalibrationCore/Source/CameraCalibrationCore/Private/SphericalLensDistortionModelHandler.cpp
--- a/Engine/Plugins/VirtualProduction/CameraCalibrationCore/Source/CameraCalibrationCore/Private/SphericalLensDistortionModelHandler.cpp
+++ b/Engine/Plugins/VirtualProduction/CameraCalibrationCore/Source/CameraCalibrationCore/Private/SphericalLensDistortionModelHandler.cpp
@@ -23,7 +23,8 @@ FVector2D USphericalLensDistortionModelHandler::ComputeDistortedUV(const FVector
 	const FVector2D OriginalDistance = NormalizedDistanceFromImageCenter;
 
 	// Iterative approach to distort an undistorted UV using coefficients that were designed to undistort
-	for (int32 Index = 0; Index < 2; ++Index)
+	constexpr uint32 NumIterations = 5;
+	for (int32 Index = 0; Index < NumIterations; ++Index)
 	{
 		const FVector2D DistanceSquared = NormalizedDistanceFromImageCenter * NormalizedDistanceFromImageCenter;
 		const float RSquared = DistanceSquared.X + DistanceSquared.Y;
diff --git a/Engine/Plugins/VirtualProduction/LensComponent/Source/LensComponent/Private/LensComponent.cpp b/Engine/Plugins/VirtualProduction/LensComponent/Source/LensComponent/Private/LensComponent.cpp
--- a/Engine/Plugins/VirtualProduction/LensComponent/Source/LensComponent/Private/LensComponent.cpp
+++ b/Engine/Plugins/VirtualProduction/LensComponent/Source/LensComponent/Private/LensComponent.cpp
@@ -186,7 +186,7 @@ void ULensComponent::TickComponent(float DeltaTime, ELevelTick TickType, FActorC
 		{
 			if (LensDistortionHandler && bWasDistortionEvaluated)
 			{
-				if (DistortionRenderingMode == EDistortionRenderingMode::LegacyPPM)
+				if (DistortionRenderingMode == EDistortionRenderingMode::PostProcessMaterial)
 				{
 					// Get the current distortion MID from the lens distortion handler
 					UMaterialInstanceDynamic* NewDistortionMID = LensDistortionHandler->GetDistortionMID();
@@ -201,16 +201,6 @@ void ULensComponent::TickComponent(float DeltaTime, ELevelTick TickType, FActorC
 					// Cache the latest distortion MID
 					LastDistortionMID = NewDistortionMID;
 
-					// Get the overscan factor and use it to modify the target camera's FOV
-					const float OverscanFactor = LensDistortionHandler->GetOverscanFactor();
-					const float OverscanSensorWidth = GetDesqueezedSensorWidth(CineCameraComponent) * OverscanFactor;
-					const float OverscanFOV = FMath::RadiansToDegrees(2.0f * FMath::Atan(OverscanSensorWidth / (2.0f * OriginalFocalLength)));
-					CineCameraComponent->SetFieldOfView(OverscanFOV);
-
-					// Update the minimum and maximum focal length of the camera (if needed)
-					CineCameraComponent->LensSettings.MinFocalLength = FMath::Min(CineCameraComponent->LensSettings.MinFocalLength, CineCameraComponent->CurrentFocalLength);
-					CineCameraComponent->LensSettings.MaxFocalLength = FMath::Max(CineCameraComponent->LensSettings.MaxFocalLength, CineCameraComponent->CurrentFocalLength);
-
 					bIsDistortionSetup = true;
 				}
 				else if (DistortionRenderingMode == EDistortionRenderingMode::SceneViewExtension)
@@ -225,6 +215,14 @@ void ULensComponent::TickComponent(float DeltaTime, ELevelTick TickType, FActorC
 
 					bIsDistortionSetup = true;
 				}
+
+				// Set the camera's overscan settings
+				if (bOverrideCameraOverscan)
+				{
+					CineCameraComponent->Overscan = LensDistortionHandler->GetOverscanFactor() - 1.0f;
+					CineCameraComponent->bScaleResolutionWithOverscan = true;
+					CineCameraComponent->bCropOverscan = (DistortionRenderingMode == EDistortionRenderingMode::SceneViewExtension);
+				}
 			}
 			else
 			{
@@ -768,13 +766,13 @@ void ULensComponent::CleanupDistortion(UCineCameraComponent* const CineCameraCom
 			LastDistortionMID = nullptr;
 		}
 
-		// Restore the original FOV of the target camera
-		const float UndistortedFOV = FMath::RadiansToDegrees(2.0f * FMath::Atan(GetDesqueezedSensorWidth(CineCameraComponent) / (2.0f * OriginalFocalLength)));
-		CineCameraComponent->SetFieldOfView(UndistortedFOV);
-
-		// Update the minimum and maximum focal length of the camera (if needed)
-		CineCameraComponent->LensSettings.MinFocalLength = FMath::Min(CineCameraComponent->LensSettings.MinFocalLength, CineCameraComponent->CurrentFocalLength);
-		CineCameraComponent->LensSettings.MaxFocalLength = FMath::Max(CineCameraComponent->LensSettings.MaxFocalLength, CineCameraComponent->CurrentFocalLength);
+		// Reset the camera's overscan settings
+		if (bOverrideCameraOverscan)
+		{
+			CineCameraComponent->Overscan = 0.0;
+			CineCameraComponent->bScaleResolutionWithOverscan = false;
+			CineCameraComponent->bCropOverscan = false;
+		}
 
 		if (UCameraCalibrationSubsystem* SubSystem = GEngine->GetEngineSubsystem<UCameraCalibrationSubsystem>())
 		{
diff --git a/Engine/Plugins/VirtualProduction/LensComponent/Source/LensComponent/Public/LensComponent.h b/Engine/Plugins/VirtualProduction/LensComponent/Source/LensComponent/Public/LensComponent.h
--- a/Engine/Plugins/VirtualProduction/LensComponent/Source/LensComponent/Public/LensComponent.h
+++ b/Engine/Plugins/VirtualProduction/LensComponent/Source/LensComponent/Public/LensComponent.h
@@ -59,8 +59,8 @@ enum class EDistortionSource : uint8
 UENUM(BlueprintType)
 enum class EDistortionRenderingMode : uint8
 {
-	/** Use the legacy post process material */
-	LegacyPPM,
+	/** Use the plugin post process material */
+	PostProcessMaterial,
 	/** Use the experimental lens distortion scene view extension. Further control of where distortion is rendered can be set via the console command r.TSR.LensDistortion */
 	SceneViewExtension UMETA(DisplayName = "Scene View Extension (Experimental)"),
 };
@@ -295,7 +295,14 @@ protected:
 
 	/** Specifies how the distortion should be rendered in the post-processing pipeline */
 	UPROPERTY(EditAnywhere, Category = "Distortion", meta = (EditCondition = "bApplyDistortion"))
-	EDistortionRenderingMode DistortionRenderingMode = EDistortionRenderingMode::LegacyPPM;
+	EDistortionRenderingMode DistortionRenderingMode = EDistortionRenderingMode::PostProcessMaterial;
+
+	/**
+	 * If checked, the camera's overscan settings (including resolution scaling and crop) will be driven by the lens component.
+	 * The overscan value is computed to be the minimum amount needed for the current distortion state. 
+	 */
+	UPROPERTY(EditAnywhere, BlueprintReadWrite, Category = "Distortion")
+	bool bOverrideCameraOverscan = false;
 
 	/** The current lens model used for distortion */
 	UPROPERTY(EditAnywhere, Category = "Distortion", meta = (EditCondition = "DistortionStateSource == EDistortionSource::Manual"))
