CameraCalibration: Expose LensFile evaluation inputs to be set manually in the details panel and from BP

diff --git a/Engine/Plugins/VirtualProduction/CameraCalibrationCore/Source/CameraCalibrationCore/Public/LensFile.h b/Engine/Plugins/VirtualProduction/CameraCalibrationCore/Source/CameraCalibrationCore/Public/LensFile.h
--- a/Engine/Plugins/VirtualProduction/CameraCalibrationCore/Source/CameraCalibrationCore/Public/LensFile.h
+++ b/Engine/Plugins/VirtualProduction/CameraCalibrationCore/Source/CameraCalibrationCore/Public/LensFile.h
@@ -381,16 +381,16 @@ struct FLensFileEvaluationInputs
 	GENERATED_BODY()
 
 public:
-	UPROPERTY(Interp, VisibleAnywhere, BlueprintReadOnly, Category = "Lens File")
+	UPROPERTY(Interp, EditAnywhere, BlueprintReadWrite, Category = "Lens File")
 	float Focus = 0.0f;
 
-	UPROPERTY(Interp, VisibleAnywhere, BlueprintReadOnly, Category = "Lens File")
+	UPROPERTY(Interp, EditAnywhere, BlueprintReadWrite, Category = "Lens File")
 	float Iris = 0.0f;
 
-	UPROPERTY(Interp, VisibleAnywhere, BlueprintReadOnly, Category = "Lens File")
+	UPROPERTY(Interp, EditAnywhere, BlueprintReadWrite, Category = "Lens File")
 	float Zoom = 0.0f;
 
-	UPROPERTY(Interp, VisibleAnywhere, BlueprintReadOnly, Category = "Lens File")
+	UPROPERTY(Interp, EditAnywhere, BlueprintReadWrite, Category = "Lens File")
 	FCameraFilmbackSettings Filmback;
 
 	UPROPERTY(Transient, BlueprintReadOnly, Category = "Lens File")
diff --git a/Engine/Plugins/VirtualProduction/LensComponent/Source/LensComponent/Private/LensComponent.cpp b/Engine/Plugins/VirtualProduction/LensComponent/Source/LensComponent/Private/LensComponent.cpp
--- a/Engine/Plugins/VirtualProduction/LensComponent/Source/LensComponent/Private/LensComponent.cpp
+++ b/Engine/Plugins/VirtualProduction/LensComponent/Source/LensComponent/Private/LensComponent.cpp
@@ -607,6 +607,13 @@ const FLensFileEvaluationInputs& ULensComponent::GetLensFileEvaluationInputs() c
 	return EvalInputs;
 }
 
+void ULensComponent::SetLensFileEvaluationInputs(float InFocus, float InZoom)
+{
+	EvalInputs.Focus = InFocus;
+	EvalInputs.Zoom = InZoom;
+	EvalInputs.bIsValid = true;
+}
+
 bool ULensComponent::ShouldApplyNodalOffsetOnTick() const
 {
 	return bApplyNodalOffsetOnTick;
@@ -925,9 +932,9 @@ void ULensComponent::UpdateLensFileEvaluationInputs(UCineCameraComponent* CineCa
 		EvalInputs.Zoom = OriginalFocalLength;
 		EvalInputs.bIsValid = true;
 	}
-	else if (EvaluationMode == EFIZEvaluationMode::UseRecordedValues)
+	else if ((EvaluationMode == EFIZEvaluationMode::UseRecordedValues) || (EvaluationMode == EFIZEvaluationMode::Manual))
 	{
-		// Do nothing, the values for EvalInputs.Focus and EvalInputs.Zoom are already loaded from the recorded sequence
+		// Do nothing, the values for EvalInputs.Focus and EvalInputs.Zoom are already valid
 		EvalInputs.bIsValid = true;
 	}
 }
diff --git a/Engine/Plugins/VirtualProduction/LensComponent/Source/LensComponent/Public/LensComponent.h b/Engine/Plugins/VirtualProduction/LensComponent/Source/LensComponent/Public/LensComponent.h
--- a/Engine/Plugins/VirtualProduction/LensComponent/Source/LensComponent/Public/LensComponent.h
+++ b/Engine/Plugins/VirtualProduction/LensComponent/Source/LensComponent/Public/LensComponent.h
@@ -27,6 +27,8 @@ enum class EFIZEvaluationMode : uint8
 	UseCameraSettings,
 	/** Evaluate the Lens File using values recorded in a level sequence (set automatically when the sequence is opened) */
 	UseRecordedValues,
+	/** Evaluate the Lens File using values set directly in the details panel or via BP/scripting */
+	Manual,
 	/** Do not evaluate the Lens File */
 	DoNotEvaluate,
 };
@@ -199,6 +201,10 @@ public:
 	UFUNCTION(BlueprintPure, Category = "Lens Component")
 	const FLensFileEvaluationInputs& GetLensFileEvaluationInputs() const;
 
+	/** Set the data used by this component to evaluate the LensFile */
+	UFUNCTION(BlueprintCallable, Category = "Lens Component")
+	void SetLensFileEvaluationInputs(float InFocus, float InZoom);
+
 	/** Returns true if nodal offset was applied during the current tick, false otherwise */
 	UFUNCTION(BlueprintPure, Category = "Lens Component")
 	bool WasNodalOffsetAppliedThisTick() const;
@@ -282,7 +288,7 @@ protected:
 	FComponentReference TargetCameraComponent;
 
 	/** Inputs to LensFile evaluation */
-	UPROPERTY(VisibleAnywhere, AdvancedDisplay, Category = "Lens File")
+	UPROPERTY(EditAnywhere, AdvancedDisplay, Category = "Lens File")
 	FLensFileEvaluationInputs EvalInputs;
 
 	/** Specifies from where the distortion state information comes */
