CameraCalibration: Change overscan override property to default to true

diff --git a/Engine/Plugins/VirtualProduction/LensComponent/Source/LensComponent/Private/LensComponent.cpp b/Engine/Plugins/VirtualProduction/LensComponent/Source/LensComponent/Private/LensComponent.cpp
--- a/Engine/Plugins/VirtualProduction/LensComponent/Source/LensComponent/Private/LensComponent.cpp
+++ b/Engine/Plugins/VirtualProduction/LensComponent/Source/LensComponent/Private/LensComponent.cpp
@@ -219,9 +219,18 @@ void ULensComponent::TickComponent(float DeltaTime, ELevelTick TickType, FActorC
 				// Set the camera's overscan settings
 				if (bOverrideCameraOverscan)
 				{
-					CineCameraComponent->Overscan = LensDistortionHandler->GetOverscanFactor() - 1.0f;
-					CineCameraComponent->bScaleResolutionWithOverscan = true;
-					CineCameraComponent->bCropOverscan = (DistortionRenderingMode == EDistortionRenderingMode::SceneViewExtension);
+					const float DesiredOverscan = FMath::Clamp(LensDistortionHandler->GetOverscanFactor() - 1.0f, 0.0f, 1.0f);
+
+					// Only override the camera overscan if the existing amount is too small for lens distortion
+					if (CineCameraComponent->Overscan < DesiredOverscan)
+					{
+						CineCameraComponent->Overscan = DesiredOverscan;
+					}
+
+					// The overscan factor used by the post process material needs to be equivalent to what is set on the camera in order for the PPM and SVE paths to render the same
+					LensDistortionHandler->SetOverscanFactor(CineCameraComponent->Overscan + 1.0f);
+
+ 					CineCameraComponent->bCropOverscan = (DistortionRenderingMode == EDistortionRenderingMode::SceneViewExtension);
 				}
 			}
 			else
@@ -773,14 +782,6 @@ void ULensComponent::CleanupDistortion(UCineCameraComponent* const CineCameraCom
 			LastDistortionMID = nullptr;
 		}
 
-		// Reset the camera's overscan settings
-		if (bOverrideCameraOverscan)
-		{
-			CineCameraComponent->Overscan = 0.0;
-			CineCameraComponent->bScaleResolutionWithOverscan = false;
-			CineCameraComponent->bCropOverscan = false;
-		}
-
 		if (UCameraCalibrationSubsystem* SubSystem = GEngine->GetEngineSubsystem<UCameraCalibrationSubsystem>())
 		{
 			if (ACameraActor* CameraActor = Cast<ACameraActor>(GetOwner()))
diff --git a/Engine/Plugins/VirtualProduction/LensComponent/Source/LensComponent/Public/LensComponent.h b/Engine/Plugins/VirtualProduction/LensComponent/Source/LensComponent/Public/LensComponent.h
--- a/Engine/Plugins/VirtualProduction/LensComponent/Source/LensComponent/Public/LensComponent.h
+++ b/Engine/Plugins/VirtualProduction/LensComponent/Source/LensComponent/Public/LensComponent.h
@@ -304,11 +304,14 @@ protected:
 	EDistortionRenderingMode DistortionRenderingMode = EDistortionRenderingMode::PostProcessMaterial;
 
 	/**
-	 * If checked, the camera's overscan settings (including resolution scaling and crop) will be driven by the lens component.
-	 * The overscan value is computed to be the minimum amount needed for the current distortion state. 
+	 * If checked, the camera's overscan value will be driven by the lens component to automatically compensate for distortion. 
+	 * The camera's overscan crop property will also be driven based on the distortion rendering mode:
+	 *   Disabled for Post Process Material
+	 *   Enabled for Scene View Extension
+	 * Note: The camera's overscan properties will not be automatically reset when the "Apply Distortion" or "Override Camera Overscan" properties are disabled
 	 */
 	UPROPERTY(EditAnywhere, BlueprintReadWrite, Category = "Distortion")
-	bool bOverrideCameraOverscan = false;
+	bool bOverrideCameraOverscan = true;
 
 	/** The current lens model used for distortion */
 	UPROPERTY(EditAnywhere, Category = "Distortion", meta = (EditCondition = "DistortionStateSource == EDistortionSource::Manual"))
