CameraCalibration: Fix lens component to properly record the original transform of the tracked component in level sequences. This supports re-applying of nodal offset, but is currently broken because the FTransform property previously used was not being recorded by any property track recorder.

diff --git a/Engine/Plugins/VirtualProduction/LensComponent/Source/LensComponent/Private/LensComponent.cpp b/Engine/Plugins/VirtualProduction/LensComponent/Source/LensComponent/Private/LensComponent.cpp
--- a/Engine/Plugins/VirtualProduction/LensComponent/Source/LensComponent/Private/LensComponent.cpp
+++ b/Engine/Plugins/VirtualProduction/LensComponent/Source/LensComponent/Private/LensComponent.cpp
@@ -397,7 +397,7 @@ void ULensComponent::ReapplyNodalOffset()
 	if (TrackedComponent.IsValid())
 	{
 		// Reset the tracked component back to its original relative transform (before nodal offset was originally applied)
-		TrackedComponent->SetRelativeTransform(OriginalTrackedComponentTransform);
+		TrackedComponent->SetRelativeTransform(GetOriginalTrackedComponentTransform());
 
 		// Now, reapply the nodal offset to the tracked component
 		ApplyNodalOffset();
@@ -429,7 +429,7 @@ void ULensComponent::ApplyNodalOffset()
 	LensFile->EvaluateNodalPointOffset(EvalInputs.Focus, EvalInputs.Zoom, Offset);
 
 	// Cache the original transform before applying the offset, so that nodal offset can potentially be re-evaluated in the future
-	OriginalTrackedComponentTransform = TrackedComponent.Get()->GetRelativeTransform();
+	SetOriginalTrackedComponentTransform(TrackedComponent.Get()->GetRelativeTransform());
 
 	TrackedComponent.Get()->AddLocalOffset(Offset.LocationOffset);
 	TrackedComponent.Get()->AddLocalRotation(Offset.RotationOffset);
@@ -473,12 +473,36 @@ void ULensComponent::ApplyNodalOffset(USceneComponent* ComponentToOffset, bool b
 	}
 
 	// Cache the original transform before applying the offset, so that nodal offset can potentially be re-evaluated in the future
-	OriginalTrackedComponentTransform = ComponentToOffset->GetRelativeTransform();
+	SetOriginalTrackedComponentTransform(ComponentToOffset->GetRelativeTransform());
 
 	ComponentToOffset->AddLocalOffset(Offset.LocationOffset);
 	ComponentToOffset->AddLocalRotation(Offset.RotationOffset);
 }
 
+FTransform ULensComponent::GetOriginalTrackedComponentTransform()
+{
+	OriginalTrackedComponentTransform.SetLocation(OriginalTrackedComponentLocation);
+
+	FRotator TrackedRotator;
+	TrackedRotator.Pitch = OriginalTrackedComponentRotation.Y;
+	TrackedRotator.Yaw = OriginalTrackedComponentRotation.Z;
+	TrackedRotator.Roll = OriginalTrackedComponentRotation.X;
+
+	OriginalTrackedComponentTransform.SetRotation(TrackedRotator.Quaternion());
+	return OriginalTrackedComponentTransform;
+}
+
+void ULensComponent::SetOriginalTrackedComponentTransform(const FTransform& NewTransform)
+{
+	OriginalTrackedComponentTransform = NewTransform;
+	OriginalTrackedComponentLocation = NewTransform.GetLocation();
+
+	FRotator NewRotator = NewTransform.GetRotation().Rotator();
+	OriginalTrackedComponentRotation.X = NewRotator.Roll;
+	OriginalTrackedComponentRotation.Y = NewRotator.Pitch;
+	OriginalTrackedComponentRotation.Z = NewRotator.Yaw;
+}
+
 void ULensComponent::EvaluateFocalLength(UCineCameraComponent* CineCameraComponent)
 {
 	FFocalLengthInfo FocalLengthInfo;
diff --git a/Engine/Plugins/VirtualProduction/LensComponent/Source/LensComponent/Public/LensComponent.h b/Engine/Plugins/VirtualProduction/LensComponent/Source/LensComponent/Public/LensComponent.h
--- a/Engine/Plugins/VirtualProduction/LensComponent/Source/LensComponent/Public/LensComponent.h
+++ b/Engine/Plugins/VirtualProduction/LensComponent/Source/LensComponent/Public/LensComponent.h
@@ -274,6 +274,12 @@ private:
 	/** Returns the sensor width of the input CineCamera, factoring in its squeeze factor */
 	float GetDesqueezedSensorWidth(UCineCameraComponent* const CineCameraComponent) const;
 
+	/** Get the original transform of the tracked component (rebuilt from the location and rotation vectors) */
+	FTransform GetOriginalTrackedComponentTransform();
+
+	/** Set the original transform of the tracked component (and also the location and rotation vectors) */
+	void SetOriginalTrackedComponentTransform(const FTransform& NewTransform);
+
 protected:
 	/** Lens File used to drive distortion with current camera settings */
 	UPROPERTY(EditAnywhere, Category="Lens File", meta=(ShowOnlyInnerProperties))
@@ -344,6 +350,17 @@ protected:
 	UPROPERTY(EditAnywhere, Category="Nodal Offset")
 	bool bApplyNodalOffsetOnTick = true;
 
+	/*
+	 * Location and Rotation of the TrackedComponent prior to nodal offset being applied 
+	 * Note: These are marked Interp so that they will be recorded in a level sequence to support re-applying nodal offset
+	 * However, recording of FTransform properties is not currently supported by the transform track recorder.
+	 * FRotator and FQuat are also not supported by the basic property track recorder, but FVector is, so we use that for both location and rotation.
+	 */
+	UPROPERTY(Interp, Category = "Nodal Offset", meta = (EditCondition=false, EditConditionHides))
+	FVector OriginalTrackedComponentLocation;
+	UPROPERTY(Interp, Category = "Nodal Offset", meta = (EditCondition=false, EditConditionHides))
+	FVector OriginalTrackedComponentRotation;
+
 	/** Serialized transform of the TrackedComponent prior to nodal offset being applied */
 	UPROPERTY(Interp, VisibleAnywhere, AdvancedDisplay, Category = "Nodal Offset")
 	FTransform OriginalTrackedComponentTransform;
