LensComponent: Make the Lens Distortion Scene View Extension the default distortion rendering mode.

diff --git a/Engine/Plugins/VirtualProduction/LensComponent/Source/LensComponent/Private/LensComponent.cpp b/Engine/Plugins/VirtualProduction/LensComponent/Source/LensComponent/Private/LensComponent.cpp
--- a/Engine/Plugins/VirtualProduction/LensComponent/Source/LensComponent/Private/LensComponent.cpp
+++ b/Engine/Plugins/VirtualProduction/LensComponent/Source/LensComponent/Private/LensComponent.cpp
@@ -14,6 +14,7 @@
 #include "Roles/LiveLinkCameraRole.h"
 #include "Roles/LiveLinkTransformRole.h"
 #include "UObject/UE5MainStreamObjectVersion.h"
+#include "UObject/FortniteMainBranchObjectVersion.h"
 
 DEFINE_LOG_CATEGORY_STATIC(LogLensComponent, Log, All);
 
@@ -322,6 +323,25 @@ void ULensComponent::PostLoad()
 #endif //WITH_EDITOR
 }
 
+void ULensComponent::Serialize(FArchive& Ar)
+{
+	Ar.UsingCustomVersion(FFortniteMainBranchObjectVersion::GUID);
+
+#if WITH_EDITOR
+	if (Ar.IsLoading())
+	{
+		// For backwards compatibility, preserve the old default value
+		if ((Ar.CustomVer(FFortniteMainBranchObjectVersion::GUID) < FFortniteMainBranchObjectVersion::LensComponentDefaultToDistortionSVE)
+			&& !GetArchetype()->IsInBlueprint())
+		{
+			DistortionRenderingMode = EDistortionRenderingMode::PostProcessMaterial;
+		}
+	}
+#endif
+
+	Super::Serialize(Ar);
+}
+
 #if WITH_EDITOR
 void ULensComponent::PostEditChangeProperty(struct FPropertyChangedEvent& PropertyChangedEvent)
 {
diff --git a/Engine/Plugins/VirtualProduction/LensComponent/Source/LensComponent/Public/LensComponent.h b/Engine/Plugins/VirtualProduction/LensComponent/Source/LensComponent/Public/LensComponent.h
--- a/Engine/Plugins/VirtualProduction/LensComponent/Source/LensComponent/Public/LensComponent.h
+++ b/Engine/Plugins/VirtualProduction/LensComponent/Source/LensComponent/Public/LensComponent.h
@@ -65,8 +65,13 @@ enum class EDistortionRenderingMode : uint8
 {
 	/** Use the plugin post process material */
 	PostProcessMaterial,
-	/** Use the experimental lens distortion scene view extension. Further control of where distortion is rendered can be set via the console command r.TSR.LensDistortion */
-	SceneViewExtension UMETA(DisplayName = "Scene View Extension (Experimental)"),
+	/**
+	  * Use the lens distortion scene view extension.
+	  * Use the console command r.TSR.LensDistortion for further control over when distortion occurs in the post-process pipeline.
+	  * r.TSR.LensDistortion = 1, distortion will be applied during TSR (default)
+	  * r.TSR.LensDistortion = 0, distortion will be applied during PrimaryUpscale
+	  */
+	SceneViewExtension,
 };
 
 /** Component for applying a post-process lens distortion effect to a CineCameraComponent on the same actor */
@@ -92,6 +97,8 @@ public:
 	UE_API virtual void PostEditImport() override;
 	UE_API virtual void PostLoad() override;
 
+	UE_API virtual void Serialize(FArchive& Ar) override;
+
 #if WITH_EDITOR
 	UE_API virtual void PostEditChangeProperty(struct FPropertyChangedEvent& PropertyChangedEvent) override;
 #endif //WITH_EDITOR
@@ -309,7 +316,7 @@ protected:
 
 	/** Specifies how the distortion should be rendered in the post-processing pipeline */
 	UPROPERTY(EditAnywhere, Category = "Distortion", meta = (EditCondition = "bApplyDistortion"))
-	EDistortionRenderingMode DistortionRenderingMode = EDistortionRenderingMode::PostProcessMaterial;
+	EDistortionRenderingMode DistortionRenderingMode = EDistortionRenderingMode::SceneViewExtension;
 
 	/**
 	 * If checked, the camera's overscan value will be driven by the lens component to automatically compensate for distortion. 
diff --git a/Engine/Source/Runtime/Core/Public/UObject/FortniteMainBranchObjectVersions.inl b/Engine/Source/Runtime/Core/Public/UObject/FortniteMainBranchObjectVersions.inl
--- a/Engine/Source/Runtime/Core/Public/UObject/FortniteMainBranchObjectVersions.inl
+++ b/Engine/Source/Runtime/Core/Public/UObject/FortniteMainBranchObjectVersions.inl
@@ -649,5 +649,8 @@ FORTNITE_MAIN_VERSION(MorphTargetCookedCPUDataCompressed, 211)
 // AnimNext variables converted to references
 FORTNITE_MAIN_VERSION(AnimNextVariableReferences, 212)
 
+// The default distortion rendering mode used by the Lens Component is now the Lens Distortion Scene View Extension
+FORTNITE_MAIN_VERSION(LensComponentDefaultToDistortionSVE, 213)
+
 // -----<new versions can be added above this line>-------------------------------------------------
 #undef FORTNITE_MAIN_VERSION
\ No newline at end of file
