CameraCalibration: Add default distortion rendering mode option.

It makes the selection for you depending on whether the lens model is D-U or U-D in nature. If D-U then PPM tends to be more accurate than SVE, because it doesn't invert the distortion equation in the material shader code.

Lens models / handlers are free to specify a different preference, either because of implementation details or because only one method is supported and not the other.

#rnx

diff --git a/Engine/Plugins/VirtualProduction/CameraCalibrationCore/Source/CameraCalibrationCore/Private/LensDistortionModelHandlerBase.cpp b/Engine/Plugins/VirtualProduction/CameraCalibrationCore/Source/CameraCalibrationCore/Private/LensDistortionModelHandlerBase.cpp
--- a/Engine/Plugins/VirtualProduction/CameraCalibrationCore/Source/CameraCalibrationCore/Private/LensDistortionModelHandlerBase.cpp
+++ b/Engine/Plugins/VirtualProduction/CameraCalibrationCore/Source/CameraCalibrationCore/Private/LensDistortionModelHandlerBase.cpp
@@ -27,6 +27,12 @@ ULensDistortionModelHandlerBase::ULensDistortionModelHandlerBase()
 	}
 }
 
+EDistortionRenderingMode ULensDistortionModelHandlerBase::GetPreferredRenderingMode() const
+{
+	// By default, reverse distorting models prefer PPM for accuracy, forward distorting models prefer SVE
+	return IsForwardDistorting() ? EDistortionRenderingMode::SceneViewExtension : EDistortionRenderingMode::PostProcessMaterial;
+}
+
 bool ULensDistortionModelHandlerBase::IsModelSupported(const TSubclassOf<ULensModel>& ModelToSupport) const
 {
 	return (LensModelClass == ModelToSupport);
diff --git a/Engine/Plugins/VirtualProduction/CameraCalibrationCore/Source/CameraCalibrationCore/Public/CameraCalibrationTypes.h b/Engine/Plugins/VirtualProduction/CameraCalibrationCore/Source/CameraCalibrationCore/Public/CameraCalibrationTypes.h
--- a/Engine/Plugins/VirtualProduction/CameraCalibrationCore/Source/CameraCalibrationCore/Public/CameraCalibrationTypes.h
+++ b/Engine/Plugins/VirtualProduction/CameraCalibrationCore/Source/CameraCalibrationCore/Public/CameraCalibrationTypes.h
@@ -3,11 +3,30 @@
 #pragma once
 
 #include "CineCameraComponent.h"
-#include "LensDistortionModelHandlerBase.h"
+#include "LensData.h"
 #include "Misc/Guid.h"
 
 #include "CameraCalibrationTypes.generated.h"
 
+/** Specifies how the distortion should be rendered in the post-processing pipeline */
+UENUM(BlueprintType)
+enum class EDistortionRenderingMode : uint8
+{
+	/** Use the plugin post process material */
+	PostProcessMaterial,
+
+	/**
+	  * Use the lens distortion scene view extension.
+	  * Use the console command r.TSR.LensDistortion for further control over when distortion occurs in the post-process pipeline.
+	  * r.TSR.LensDistortion = 1, distortion will be applied during TSR (default)
+	  * r.TSR.LensDistortion = 0, distortion will be applied during PrimaryUpscale
+	  */
+	SceneViewExtension,
+
+	/** Use the lens model's preferred rendering mode */
+	Preferred,
+};
+
 /** Results from a distortion calibration, including camera intrinsics and either the parameters to an analytical model or an ST Map */
 USTRUCT(BlueprintType)
 struct FDistortionCalibrationResult
diff --git a/Engine/Plugins/VirtualProduction/CameraCalibrationCore/Source/CameraCalibrationCore/Public/LensDistortionModelHandlerBase.h b/Engine/Plugins/VirtualProduction/CameraCalibrationCore/Source/CameraCalibrationCore/Public/LensDistortionModelHandlerBase.h
--- a/Engine/Plugins/VirtualProduction/CameraCalibrationCore/Source/CameraCalibrationCore/Public/LensDistortionModelHandlerBase.h
+++ b/Engine/Plugins/VirtualProduction/CameraCalibrationCore/Source/CameraCalibrationCore/Public/LensDistortionModelHandlerBase.h
@@ -2,6 +2,7 @@
 
 #pragma once
 
+#include "CameraCalibrationTypes.h"
 #include "CineCameraSettings.h"
 #include "CoreMinimal.h"
 #include "LensData.h"
@@ -138,6 +139,14 @@ public:
 	/** Returns true if this model performs forward distortion (needs inverse overscan), false if reverse distortion (needs forward overscan) */
 	virtual bool IsForwardDistorting() const { return true; }
 
+	/** 
+	 * Returns the preferred rendering mode for this lens model.
+	 * By default, reverse distorting models prefer PPM for accuracy, forward distorting models prefer SVE.
+	 * Override this to customize the preference for specific lens models.
+	 * Return Preferred to use the default behavior based on IsForwardDistorting().
+	 */
+	UE_API virtual EDistortionRenderingMode GetPreferredRenderingMode() const;
+
 protected:
 	/** Initialize the handler. Derived classes must set the LensModelClass that they support, if not already set */
 	virtual void InitializeHandler() PURE_VIRTUAL(ULensDistortionModelHandlerBase::InitializeHandler);
diff --git a/Engine/Plugins/VirtualProduction/LensComponent/Config/BaseLensComponent.ini b/Engine/Plugins/VirtualProduction/LensComponent/Config/BaseLensComponent.ini
--- a/Engine/Plugins/VirtualProduction/LensComponent/Config/BaseLensComponent.ini
+++ b/Engine/Plugins/VirtualProduction/LensComponent/Config/BaseLensComponent.ini
@@ -10,3 +10,11 @@
 +ClassRedirects=(OldName="/Script/CameraCalibrationCoreMovieScene.MovieSceneLensComponentSection", NewName="/Script/LensComponentEditor.MovieSceneLensComponentSection")
 +ClassRedirects=(OldName="/Script/CameraCalibrationCoreMovieScene.MovieSceneLensComponentTrack", NewName="/Script/LensComponentEditor.MovieSceneLensComponentTrack")
 +ClassRedirects=(OldName="/Script/CameraCalibrationCoreSequencer.MovieSceneLensComponentTrackRecorder", NewName="/Script/LensComponentEditor.MovieSceneLensComponentTrackRecorder")
+
+;
+; 5.7 Redirects
+;
+
+; Enums
+
++EnumRedirects=(OldName="/Script/LensComponent.EDistortionRenderingMode", NewName="/Script/CameraCalibrationCore.EDistortionRenderingMode")
diff --git a/Engine/Plugins/VirtualProduction/LensComponent/Source/LensComponent/Private/LensComponent.cpp b/Engine/Plugins/VirtualProduction/LensComponent/Source/LensComponent/Private/LensComponent.cpp
--- a/Engine/Plugins/VirtualProduction/LensComponent/Source/LensComponent/Private/LensComponent.cpp
+++ b/Engine/Plugins/VirtualProduction/LensComponent/Source/LensComponent/Private/LensComponent.cpp
@@ -189,7 +189,16 @@ void ULensComponent::TickComponent(float DeltaTime, ELevelTick TickType, FActorC
 		{
 			if (LensDistortionHandler && bWasDistortionEvaluated)
 			{
-				if (DistortionRenderingMode == EDistortionRenderingMode::PostProcessMaterial)
+				const EDistortionRenderingMode ResolvedRenderingMode = GetDistortionRenderingMode();
+				
+				// Check if the resolved rendering mode has changed (can happen with Preferred mode)
+				if (bIsDistortionSetup && ResolvedRenderingMode != LastAppliedRenderingMode)
+				{
+					// Clean up the previous mode before switching
+					CleanupDistortion(CineCameraComponent);
+				}
+				
+				if (ResolvedRenderingMode == EDistortionRenderingMode::PostProcessMaterial)
 				{
 					// Get the current distortion MID from the lens distortion handler
 					UMaterialInstanceDynamic* NewDistortionMID = LensDistortionHandler->GetDistortionMID();
@@ -205,8 +214,9 @@ void ULensComponent::TickComponent(float DeltaTime, ELevelTick TickType, FActorC
 					LastDistortionMID = NewDistortionMID;
 
 					bIsDistortionSetup = true;
+					LastAppliedRenderingMode = ResolvedRenderingMode;
 				}
-				else if (DistortionRenderingMode == EDistortionRenderingMode::SceneViewExtension)
+				else if (ResolvedRenderingMode == EDistortionRenderingMode::SceneViewExtension)
 				{
 					if (UCameraCalibrationSubsystem* SubSystem = GEngine->GetEngineSubsystem<UCameraCalibrationSubsystem>())
 					{
@@ -217,6 +227,7 @@ void ULensComponent::TickComponent(float DeltaTime, ELevelTick TickType, FActorC
 					}
 
 					bIsDistortionSetup = true;
+					LastAppliedRenderingMode = ResolvedRenderingMode;
 				}
 
 				// Set the camera's overscan settings
@@ -228,7 +239,7 @@ void ULensComponent::TickComponent(float DeltaTime, ELevelTick TickType, FActorC
 					// The overscan factor used by the post process material needs to be equivalent to what is set on the camera in order for the PPM and SVE paths to render the same
 					LensDistortionHandler->SetOverscanFactor(CineCameraComponent->Overscan + 1.0f);
 
- 					CineCameraComponent->bCropOverscan = (DistortionRenderingMode == EDistortionRenderingMode::SceneViewExtension);
+ 					CineCameraComponent->bCropOverscan = (ResolvedRenderingMode == EDistortionRenderingMode::SceneViewExtension);
 				}
 			}
 			else
@@ -766,6 +777,37 @@ ULensDistortionModelHandlerBase* ULensComponent::GetLensDistortionHandler() cons
 	return LensDistortionHandlerMap.FindRef(LensModel);
 }
 
+EDistortionRenderingMode ULensComponent::GetDistortionRenderingMode() const
+{
+	// If not Preferred, return as specified.
+	if (DistortionRenderingMode != EDistortionRenderingMode::Preferred)
+	{
+		return DistortionRenderingMode;
+	}
+
+	// Get the handler for the current lens model
+	const ULensDistortionModelHandlerBase* LensDistortionHandler = LensDistortionHandlerMap.FindRef(LensModel);
+	if (LensDistortionHandler)
+	{
+		// Let the lens model decide its preferred rendering mode
+		const EDistortionRenderingMode PreferredMode = LensDistortionHandler->GetPreferredRenderingMode();
+		
+		// If the handler unexpectedly returns Preferred, use default behavior based on IsForwardDistorting()
+		if (PreferredMode == EDistortionRenderingMode::Preferred)
+		{
+			// Reverse distorting models prefer PPM for accuracy, forward distorting prefer SVE
+			return LensDistortionHandler->IsForwardDistorting() 
+				? EDistortionRenderingMode::SceneViewExtension 
+				: EDistortionRenderingMode::PostProcessMaterial;
+		}
+		
+		return PreferredMode;
+	}
+
+	// Default to SVE if no handler is found
+	return EDistortionRenderingMode::SceneViewExtension;
+}
+
 PRAGMA_DISABLE_DEPRECATION_WARNINGS
 FDistortionHandlerPicker ULensComponent::GetDistortionHandlerPicker() const
 {
@@ -833,6 +875,7 @@ void ULensComponent::CleanupDistortion(UCineCameraComponent* const CineCameraCom
 	}
 
 	bIsDistortionSetup = false;
+	LastAppliedRenderingMode = EDistortionRenderingMode::Preferred;
 }
 
 void ULensComponent::ClearDistortionState()
diff --git a/Engine/Plugins/VirtualProduction/LensComponent/Source/LensComponent/Public/LensComponent.h b/Engine/Plugins/VirtualProduction/LensComponent/Source/LensComponent/Public/LensComponent.h
--- a/Engine/Plugins/VirtualProduction/LensComponent/Source/LensComponent/Public/LensComponent.h
+++ b/Engine/Plugins/VirtualProduction/LensComponent/Source/LensComponent/Public/LensComponent.h
@@ -5,6 +5,7 @@
 #include "Components/ActorComponent.h"
 
 #include "CameraCalibrationTypes.h"
+#include "LensDistortionModelHandlerBase.h"
 #include "LensFile.h"
 #include "LiveLinkComponentController.h"
 
@@ -59,21 +60,6 @@ enum class EDistortionSource : uint8
 	Manual,
 };
 
-/** Specifies how the distortion should be rendered in the post-processing pipeline */
-UENUM(BlueprintType)
-enum class EDistortionRenderingMode : uint8
-{
-	/** Use the plugin post process material */
-	PostProcessMaterial,
-	/**
-	  * Use the lens distortion scene view extension.
-	  * Use the console command r.TSR.LensDistortion for further control over when distortion occurs in the post-process pipeline.
-	  * r.TSR.LensDistortion = 1, distortion will be applied during TSR (default)
-	  * r.TSR.LensDistortion = 0, distortion will be applied during PrimaryUpscale
-	  */
-	SceneViewExtension,
-};
-
 /** Component for applying a post-process lens distortion effect to a CineCameraComponent on the same actor */
 UCLASS(MinimalAPI, HideCategories=(Tags, Activation, Cooking, AssetUserData, Collision), meta=(BlueprintSpawnableComponent))
 class ULensComponent : public UActorComponent
@@ -316,7 +302,7 @@ protected:
 
 	/** Specifies how the distortion should be rendered in the post-processing pipeline */
 	UPROPERTY(EditAnywhere, Category = "Distortion", meta = (EditCondition = "bApplyDistortion"))
-	EDistortionRenderingMode DistortionRenderingMode = EDistortionRenderingMode::SceneViewExtension;
+	EDistortionRenderingMode DistortionRenderingMode = EDistortionRenderingMode::Preferred;
 
 	/**
 	 * If checked, the camera's overscan value will be driven by the lens component to automatically compensate for distortion. 
@@ -386,6 +372,9 @@ protected:
 	UPROPERTY()
 	TObjectPtr<UMaterialInstanceDynamic> LastDistortionMID = nullptr;
 
+	/** Track the last applied rendering mode to detect changes when using Preferred mode */
+	EDistortionRenderingMode LastAppliedRenderingMode = EDistortionRenderingMode::Preferred;
+
 	/** Cached most recent target camera, used to clean up the old camera when the user changes the target */
 	UPROPERTY()
 	TObjectPtr<UCineCameraComponent> LastCameraComponent = nullptr;
@@ -441,6 +430,9 @@ private:
 
 	/** Delegate that is triggered when the LensModel changes */
 	FOnLensComponentModelChanged OnLensComponentModelChangedDelegate;
+
+	/** Resolves the distortion rendering mode to the actual mode that will be used (resolving Preferred if necessary) */
+	EDistortionRenderingMode GetDistortionRenderingMode() const;
 };
 
 #undef UE_API
