LensDistortion: Add new lens distortion option to apply distortion as a scene view extension pass after tonemap.

* This can be used to apply distortion before other scene view extensions (i.e. CompositeCore, OpenColorIO), which is otherwise not possible with post-process materials.
* Renamed previously called "SceneViewExtention" option to "TemporalSuperResolution" since its primary application is in the renderer TSR pass, not the SVE.
* We now automatically reset camera overscan defaults when the "Apply Distortion" or "Override Camera Overscan" are disabled.
* Fixed warning emitted for bScaleOverscan used as an EditCondition with visibility that did not match.

#Synced-CL 49805348

diff --git a/Engine/Plugins/VirtualProduction/CameraCalibrationCore/Config/DefaultCameraCalibrationCore.ini b/Engine/Plugins/VirtualProduction/CameraCalibrationCore/Config/DefaultCameraCalibrationCore.ini
--- a/Engine/Plugins/VirtualProduction/CameraCalibrationCore/Config/DefaultCameraCalibrationCore.ini
+++ b/Engine/Plugins/VirtualProduction/CameraCalibrationCore/Config/DefaultCameraCalibrationCore.ini
@@ -48,3 +48,7 @@
 ;
 +ClassRedirects=(OldName="/Script/GameplayCameras.LensCalibrationCameraNode", NewName="/Script/CameraCalibrationCore.LensCalibrationCameraNode")
 
+;
+; 5.8 redirects
+;
++EnumRedirects=(OldName="EDistortionRenderingMode", ValueChanges=(("SceneViewExtension", "TemporalSuperResolution")))
\ No newline at end of file
diff --git a/Engine/Plugins/VirtualProduction/CameraCalibrationCore/Source/CameraCalibrationCore/Private/CameraCalibrationSettings.cpp b/Engine/Plugins/VirtualProduction/CameraCalibrationCore/Source/CameraCalibrationCore/Private/CameraCalibrationSettings.cpp
--- a/Engine/Plugins/VirtualProduction/CameraCalibrationCore/Source/CameraCalibrationCore/Private/CameraCalibrationSettings.cpp
+++ b/Engine/Plugins/VirtualProduction/CameraCalibrationCore/Source/CameraCalibrationCore/Private/CameraCalibrationSettings.cpp
@@ -34,6 +34,9 @@ UCameraCalibrationSettings::UCameraCalibrationSettings()
 	DefaultDistortionDisplacementMaterials.Add(UBrownConradyDULensDistortionModelHandler::StaticClass(), TSoftObjectPtr<UMaterialInterface>(FSoftObjectPath(TEXT("/CameraCalibrationCore/Materials/M_BrownConradyDUDistortionDisplacementMap.M_BrownConradyDUDistortionDisplacementMap"))));
 	DefaultDistortionMaterials.Add(UBrownConradyDULensDistortionModelHandler::StaticClass(), TSoftObjectPtr<UMaterialInterface>(FSoftObjectPath(TEXT("/CameraCalibrationCore/Materials/M_DistortionPostProcess.M_DistortionPostProcess"))));
 
+	// Defaults priority such that the SVE runs before OpenColorIO and CompositeCore.
+	SceneViewExtensionPriority = 100 + 10 /** OPENCOLORIO_SCENE_VIEW_EXTENSION_PRIORITY + 10 */;
+
 #if WITH_EDITORONLY_DATA
 	auto InitOverlayOverrides = [this]()
 	{
diff --git a/Engine/Plugins/VirtualProduction/CameraCalibrationCore/Source/CameraCalibrationCore/Private/CameraCalibrationSubsystem.cpp b/Engine/Plugins/VirtualProduction/CameraCalibrationCore/Source/CameraCalibrationCore/Private/CameraCalibrationSubsystem.cpp
--- a/Engine/Plugins/VirtualProduction/CameraCalibrationCore/Source/CameraCalibrationCore/Private/CameraCalibrationSubsystem.cpp
+++ b/Engine/Plugins/VirtualProduction/CameraCalibrationCore/Source/CameraCalibrationCore/Private/CameraCalibrationSubsystem.cpp
@@ -194,9 +194,9 @@ TArray<FName> UCameraCalibrationSubsystem::GetCameraCalibrationSteps() const
 	return OutKeys;
 }
 
-void UCameraCalibrationSubsystem::SetLensDistortionSVEState(ACameraActor* CameraActor, FDisplacementMapBlendingParams DistortionState, ULensDistortionModelHandlerBase* LensDistortionHandler)
+void UCameraCalibrationSubsystem::SetLensDistortionSVEState(ACameraActor* CameraActor, FDisplacementMapBlendingParams DistortionState, ULensDistortionModelHandlerBase* LensDistortionHandler, EDistortionRenderingMode InRenderingMode)
 {
-	SceneViewExtension->UpdateDistortionState_AnyThread(CameraActor, DistortionState, LensDistortionHandler);
+	SceneViewExtension->UpdateDistortionState_AnyThread(CameraActor, DistortionState, LensDistortionHandler, InRenderingMode);
 }
 
 void UCameraCalibrationSubsystem::ClearLensDistortionSVEState(ACameraActor* CameraActor)
diff --git a/Engine/Plugins/VirtualProduction/CameraCalibrationCore/Source/CameraCalibrationCore/Private/LensDistortionModelHandlerBase.cpp b/Engine/Plugins/VirtualProduction/CameraCalibrationCore/Source/CameraCalibrationCore/Private/LensDistortionModelHandlerBase.cpp
--- a/Engine/Plugins/VirtualProduction/CameraCalibrationCore/Source/CameraCalibrationCore/Private/LensDistortionModelHandlerBase.cpp
+++ b/Engine/Plugins/VirtualProduction/CameraCalibrationCore/Source/CameraCalibrationCore/Private/LensDistortionModelHandlerBase.cpp
@@ -29,8 +29,8 @@ ULensDistortionModelHandlerBase::ULensDistortionModelHandlerBase()
 
 EDistortionRenderingMode ULensDistortionModelHandlerBase::GetPreferredRenderingMode() const
 {
-	// By default, reverse distorting models prefer PPM for accuracy, forward distorting models prefer SVE
-	return IsForwardDistorting() ? EDistortionRenderingMode::SceneViewExtension : EDistortionRenderingMode::PostProcessMaterial;
+	// By default, reverse distorting models prefer the plugin pass for accuracy, forward distorting models prefer TSR
+	return IsForwardDistorting() ? EDistortionRenderingMode::TemporalSuperResolution : EDistortionRenderingMode::PassAfterTonemap;
 }
 
 bool ULensDistortionModelHandlerBase::IsModelSupported(const TSubclassOf<ULensModel>& ModelToSupport) const
diff --git a/Engine/Plugins/VirtualProduction/CameraCalibrationCore/Source/CameraCalibrationCore/Private/LensDistortionSceneViewExtension.cpp b/Engine/Plugins/VirtualProduction/CameraCalibrationCore/Source/CameraCalibrationCore/Private/LensDistortionSceneViewExtension.cpp
--- a/Engine/Plugins/VirtualProduction/CameraCalibrationCore/Source/CameraCalibrationCore/Private/LensDistortionSceneViewExtension.cpp
+++ b/Engine/Plugins/VirtualProduction/CameraCalibrationCore/Source/CameraCalibrationCore/Private/LensDistortionSceneViewExtension.cpp
@@ -5,10 +5,12 @@
 #include "CineCameraComponent.h"
 #include "Camera/CameraActor.h"
 #include "Camera/CameraComponent.h"
+#include "CameraCalibrationSettings.h"
 #include "Containers/DynamicRHIResourceArray.h"
 #include "Engine/TextureRenderTarget2D.h"
 #include "GlobalShader.h"
 #include "PostProcess/LensDistortion.h"
+#include "PostProcess/PostProcessMaterialInputs.h"
 #include "RHIResourceUtils.h"
 #include "ScreenPass.h"
 #include "ShaderParameterStruct.h"
@@ -27,13 +29,14 @@ FLensDistortionSceneViewExtension::FLensDistortionSceneViewExtension(const FAuto
 {
 }
 
-void FLensDistortionSceneViewExtension::UpdateDistortionState_AnyThread(ACameraActor* CameraActor, FDisplacementMapBlendingParams DistortionState, ULensDistortionModelHandlerBase* LensDistortionHandler)
+void FLensDistortionSceneViewExtension::UpdateDistortionState_AnyThread(ACameraActor* CameraActor, FDisplacementMapBlendingParams DistortionState, ULensDistortionModelHandlerBase* LensDistortionHandler, EDistortionRenderingMode InRenderingMode)
 {
 	FScopeLock ScopeLock(&DistortionStateMapCriticalSection);
 
 	FCameraDistortionProxy CamProxy;
 	CamProxy.Params = DistortionState;
 	CamProxy.LensDistortionHandler = LensDistortionHandler;
+	CamProxy.RenderingMode = InRenderingMode;
 
 	if (UCameraComponent* Component = CameraActor->GetCameraComponent())
 	{
@@ -526,6 +529,11 @@ void FLensDistortionSceneViewExtension::RecenterSTDisplacementMap_RenderThread(F
 }
 
 
+int32 FLensDistortionSceneViewExtension::GetPriority() const
+{
+	return GetDefault<UCameraCalibrationSettings>()->GetSceneViewExtensionPriority();
+}
+
 void FLensDistortionSceneViewExtension::PreRenderView_RenderThread(FRDGBuilder& GraphBuilder, FSceneView& InView)
 {
 	FLensDistortionLUT ViewDistortionLUT;
@@ -545,6 +553,11 @@ bool FLensDistortionSceneViewExtension::RenderViewDistortionLUT(FRDGBuilder& Gra
 		return false;
 	}
 
+	if (Proxy->RenderingMode != EDistortionRenderingMode::TemporalSuperResolution)
+	{
+		return false;
+	}
+
 	FVector2D SensorSize = FVector2D(Proxy->FilmbackSettings.SensorWidth, Proxy->FilmbackSettings.SensorHeight);
 	FDisplacementMapBlendingParams BlendState = Proxy->Params;
 
@@ -663,3 +676,38 @@ bool FLensDistortionSceneViewExtension::RenderViewDistortionLUT(FRDGBuilder& Gra
 
 	return true;
 }
+
+void FLensDistortionSceneViewExtension::SubscribeToPostProcessingPass(EPostProcessingPass PassId, const FSceneView& InView, FAfterPassCallbackDelegateArray& InOutPassCallbacks, bool bIsPassEnabled)
+{
+	if (PassId == EPostProcessingPass::Tonemap)
+	{
+		FScopeLock ScopeLock(&DistortionStateMapCriticalSection);
+
+		const FCameraDistortionProxy* Proxy = DistortionStateMap.Find(InView.ViewActor.ActorUniqueId);
+		
+		if (Proxy && Proxy->RenderingMode == EDistortionRenderingMode::PassAfterTonemap)
+		{
+			TStrongObjectPtr<ULensDistortionModelHandlerBase> ModelHandler = Proxy->LensDistortionHandler.Pin();
+			
+			if (ModelHandler.IsValid())
+			{
+				UMaterialInterface* LensDistortionMaterial = ModelHandler->GetDistortionMID();
+
+				InOutPassCallbacks.Add(FAfterPassCallbackDelegate::CreateRaw(this, &FLensDistortionSceneViewExtension::PostProcessPassAfterTonemap_RenderThread, LensDistortionMaterial));
+			}
+		}
+	}
+}
+
+FScreenPassTexture FLensDistortionSceneViewExtension::PostProcessPassAfterTonemap_RenderThread(FRDGBuilder& GraphBuilder, const FSceneView& InView, const FPostProcessMaterialInputs& InOutInputs, UMaterialInterface* LensDistortionMaterial)
+{
+	if (IsValid(LensDistortionMaterial))
+	{
+		// Note: Plugin material could be entirely replaced by shaders should we deprecate the blendable EDistortionRenderingMode::PostProcessMaterial path.
+		return AddPostProcessMaterialPass(GraphBuilder, InView, InOutInputs, LensDistortionMaterial);
+	}
+	else
+	{
+		return InOutInputs.ReturnUntouchedSceneColorForPostProcessing(GraphBuilder);
+	}
+}
diff --git a/Engine/Plugins/VirtualProduction/CameraCalibrationCore/Source/CameraCalibrationCore/Public/CameraCalibrationSettings.h b/Engine/Plugins/VirtualProduction/CameraCalibrationCore/Source/CameraCalibrationCore/Public/CameraCalibrationSettings.h
--- a/Engine/Plugins/VirtualProduction/CameraCalibrationCore/Source/CameraCalibrationCore/Public/CameraCalibrationSettings.h
+++ b/Engine/Plugins/VirtualProduction/CameraCalibrationCore/Source/CameraCalibrationCore/Public/CameraCalibrationSettings.h
@@ -81,6 +81,9 @@ public:
 	/** Returns true if the calibration dataset import and export featuers are enabled, false otherwise */
 	bool IsCalibrationDatasetImportExportEnabled() const { return bEnableCalibrationDatasetImportExport; }
 
+	/** Returns the scene view extension priority ranking. */
+	int32 GetSceneViewExtensionPriority() const { return SceneViewExtensionPriority; }
+
 #if WITH_EDITOR
 protected:
 	FOnDisplacementMapResolutionChanged DisplacementMapResolutionChangedDelegate;
@@ -132,6 +135,10 @@ private:
 	UPROPERTY(config, EditAnywhere, Category = "Settings")
 	bool bEnableCalibrationDatasetImportExport = true;
 
+	/** Lens distortion (scene view extension) post-processing priority, which defaults to before OpenColorIO. */
+	UPROPERTY(config, EditAnywhere, Category = Settings, AdvancedDisplay)
+	int32 SceneViewExtensionPriority;
+
 private:
 	/** Delegate handle to run after the engine is initialized */
 	FDelegateHandle PostEngineInitHandle;
diff --git a/Engine/Plugins/VirtualProduction/CameraCalibrationCore/Source/CameraCalibrationCore/Public/CameraCalibrationSubsystem.h b/Engine/Plugins/VirtualProduction/CameraCalibrationCore/Source/CameraCalibrationCore/Public/CameraCalibrationSubsystem.h
--- a/Engine/Plugins/VirtualProduction/CameraCalibrationCore/Source/CameraCalibrationCore/Public/CameraCalibrationSubsystem.h
+++ b/Engine/Plugins/VirtualProduction/CameraCalibrationCore/Source/CameraCalibrationCore/Public/CameraCalibrationSubsystem.h
@@ -125,8 +125,8 @@ public:
 	/** Unregister an overlay material */
 	UE_API void UnregisterOverlayMaterial(const FName& MaterialName);
 
-	/** Adds the input distortion state and blending parameters to the Lens Distortion Scene View Extension for the input camera actor */
-	UE_API void SetLensDistortionSVEState(ACameraActor* CameraActor, FDisplacementMapBlendingParams DistortionState, ULensDistortionModelHandlerBase* LensDistortionHandler = nullptr);
+	/** Adds the input distortion state, rendering mode and blending parameters to the Lens Distortion Scene View Extension for the input camera actor. */
+	UE_API void SetLensDistortionSVEState(ACameraActor* CameraActor, FDisplacementMapBlendingParams DistortionState, ULensDistortionModelHandlerBase* LensDistortionHandler, EDistortionRenderingMode InRenderingMode);
 
 	/** Removes all distortion state from the Lens Distortion Scene View Extension for the input camera actor */
 	UE_API void ClearLensDistortionSVEState(ACameraActor* CameraActor);
diff --git a/Engine/Plugins/VirtualProduction/CameraCalibrationCore/Source/CameraCalibrationCore/Public/CameraCalibrationTypes.h b/Engine/Plugins/VirtualProduction/CameraCalibrationCore/Source/CameraCalibrationCore/Public/CameraCalibrationTypes.h
--- a/Engine/Plugins/VirtualProduction/CameraCalibrationCore/Source/CameraCalibrationCore/Public/CameraCalibrationTypes.h
+++ b/Engine/Plugins/VirtualProduction/CameraCalibrationCore/Source/CameraCalibrationCore/Public/CameraCalibrationTypes.h
@@ -12,19 +12,30 @@
 UENUM(BlueprintType)
 enum class EDistortionRenderingMode : uint8
 {
-	/** Use the plugin post process material */
-	PostProcessMaterial,
+	/** Use the lens model's preferred rendering mode */
+	Preferred = 2u,
+
+	/** Use the plugin post process material as a camera blendable. */
+	PostProcessMaterial = 0u,
 
 	/**
-	  * Use the lens distortion scene view extension.
-	  * Use the console command r.TSR.LensDistortion for further control over when distortion occurs in the post-process pipeline.
-	  * r.TSR.LensDistortion = 1, distortion will be applied during TSR (default)
-	  * r.TSR.LensDistortion = 0, distortion will be applied during PrimaryUpscale
+	  * Apply lens distortion inside the renderer's temporal super resolution.
+	  * 
+	  * Renderer passes after TSR are automatically made aware of lens distortion through renderer diplacements maps (FLensDistortionLUT).
+	  * Note that if TSR is not active or if r.TSR.LensDistortion = 0, distortion will be applied later in PrimaryUpscale instead.
 	  */
-	SceneViewExtension,
+	TemporalSuperResolution = 1u,
 
-	/** Use the lens model's preferred rendering mode */
-	Preferred,
+	/**
+	* Use the plugin post process material as a scene view extension pass after tone-mapping.
+	* 
+	* This can be used to apply distortion before other scene view extensions (i.e. CompositeCore, OpenColorIO),
+	* which is otherwise not possible with post-process materials.
+	*/
+	PassAfterTonemap = 3u,
+
+	//~ DEPRECATED
+	SceneViewExtension UE_DEPRECATED(5.8, "Renamed to TemporalSuperResolution") = TemporalSuperResolution UMETA(Hidden),
 };
 
 /** Results from a distortion calibration, including camera intrinsics and either the parameters to an analytical model or an ST Map */
diff --git a/Engine/Plugins/VirtualProduction/CameraCalibrationCore/Source/CameraCalibrationCore/Public/LensDistortionSceneViewExtension.h b/Engine/Plugins/VirtualProduction/CameraCalibrationCore/Source/CameraCalibrationCore/Public/LensDistortionSceneViewExtension.h
--- a/Engine/Plugins/VirtualProduction/CameraCalibrationCore/Source/CameraCalibrationCore/Public/LensDistortionSceneViewExtension.h
+++ b/Engine/Plugins/VirtualProduction/CameraCalibrationCore/Source/CameraCalibrationCore/Public/LensDistortionSceneViewExtension.h
@@ -19,7 +19,9 @@ public:
 	FLensDistortionSceneViewExtension(const FAutoRegister& AutoRegister);
 
 	//~ Begin ISceneViewExtension interface	
+	virtual int32 GetPriority() const override;
 	virtual void PreRenderView_RenderThread(FRDGBuilder& GraphBuilder, FSceneView& InView) override;
+	virtual void SubscribeToPostProcessingPass(EPostProcessingPass PassId, const FSceneView& InView, FAfterPassCallbackDelegateArray& InOutPassCallbacks, bool bIsPassEnabled) override;
 	//~End ISceneVIewExtension interface
 
 	/** Render forward & inverse displacement maps and return the engine lens distortion LUT. */
@@ -27,7 +29,7 @@ public:
 
 public:
 	/** Update the distortion state and blending params for the input camera */
-	void UpdateDistortionState_AnyThread(ACameraActor* CameraActor, FDisplacementMapBlendingParams DistortionState, ULensDistortionModelHandlerBase* LensDistortionHandler = nullptr);
+	void UpdateDistortionState_AnyThread(ACameraActor* CameraActor, FDisplacementMapBlendingParams DistortionState, ULensDistortionModelHandlerBase* LensDistortionHandler, EDistortionRenderingMode InRenderingMode);
 
 	/** Returns true if any distortion state has been registered. */
 	CAMERACALIBRATIONCORE_API bool HasDistortionState_AnyThread() const;
@@ -54,7 +56,9 @@ private:
 	/** Recenters the specified distortion map within an overscanned map, filling the overscanned region with zeros */
 	void RecenterSTDisplacementMap_RenderThread(FRDGBuilder& GraphBuilder, const FRDGTextureRef& InDisplacementMap, float InOverscan, FRDGTextureRef& OutRecenteredDisplacementMap) const;
 	
-	
+	/** Post-process pass where we optionally apply the distortion material. */
+	FScreenPassTexture PostProcessPassAfterTonemap_RenderThread(FRDGBuilder& GraphBuilder, const FSceneView& InView, const FPostProcessMaterialInputs& Inputs, UMaterialInterface* LensDistortionMaterial);
+
 private:
 
 	struct FCameraDistortionProxy
@@ -63,6 +67,7 @@ private:
 		float CameraOverscan = 1.0f;
 		FCameraFilmbackSettings FilmbackSettings = {};
 		TWeakObjectPtr<ULensDistortionModelHandlerBase> LensDistortionHandler = nullptr;
+		EDistortionRenderingMode RenderingMode = EDistortionRenderingMode::Preferred;
 	};
 	/** Map of cameras to their associated distortion state and blending parameters, used to determine if and how displacement maps should be rendered for a specific view */
 	TMap<uint32, FCameraDistortionProxy> DistortionStateMap;
diff --git a/Engine/Plugins/VirtualProduction/LensComponent/Source/LensComponent/Private/LensComponent.cpp b/Engine/Plugins/VirtualProduction/LensComponent/Source/LensComponent/Private/LensComponent.cpp
--- a/Engine/Plugins/VirtualProduction/LensComponent/Source/LensComponent/Private/LensComponent.cpp
+++ b/Engine/Plugins/VirtualProduction/LensComponent/Source/LensComponent/Private/LensComponent.cpp
@@ -216,13 +216,13 @@ void ULensComponent::TickComponent(float DeltaTime, ELevelTick TickType, FActorC
 					bIsDistortionSetup = true;
 					LastAppliedRenderingMode = ResolvedRenderingMode;
 				}
-				else if (ResolvedRenderingMode == EDistortionRenderingMode::SceneViewExtension)
+				else if (ResolvedRenderingMode == EDistortionRenderingMode::TemporalSuperResolution || ResolvedRenderingMode == EDistortionRenderingMode::PassAfterTonemap)
 				{
 					if (UCameraCalibrationSubsystem* SubSystem = GEngine->GetEngineSubsystem<UCameraCalibrationSubsystem>())
 					{
 						if (ACameraActor* CameraActor = Cast<ACameraActor>(GetOwner()))
 						{
-							SubSystem->SetLensDistortionSVEState(CameraActor, BlendState, LensDistortionHandler);
+							SubSystem->SetLensDistortionSVEState(CameraActor, BlendState, LensDistortionHandler, ResolvedRenderingMode);
 						}
 					}
 
@@ -239,7 +239,7 @@ void ULensComponent::TickComponent(float DeltaTime, ELevelTick TickType, FActorC
 					// The overscan factor used by the post process material needs to be equivalent to what is set on the camera in order for the PPM and SVE paths to render the same
 					LensDistortionHandler->SetOverscanFactor(CineCameraComponent->Overscan + 1.0f);
 
- 					CineCameraComponent->bCropOverscan = (ResolvedRenderingMode == EDistortionRenderingMode::SceneViewExtension);
+ 					CineCameraComponent->bCropOverscan = (ResolvedRenderingMode == EDistortionRenderingMode::TemporalSuperResolution);
 				}
 			}
 			else
@@ -362,13 +362,7 @@ void ULensComponent::PostEditChangeProperty(struct FPropertyChangedEvent& Proper
 
 	if (PropertyName == GET_MEMBER_NAME_CHECKED(ULensComponent, bApplyDistortion))
 	{
-		if (UCineCameraComponent* const CineCameraComponent = Cast<UCineCameraComponent>(TargetCameraComponent.GetComponent(GetOwner())))
-		{
-			if (!bApplyDistortion)
-			{
-				CleanupDistortion(CineCameraComponent);
-			}
-		}
+		SetApplyDistortion(bApplyDistortion);
 	}
 	else if (PropertyName == GET_MEMBER_NAME_CHECKED(ULensComponent, DistortionRenderingMode))
 	{
@@ -399,6 +393,10 @@ void ULensComponent::PostEditChangeProperty(struct FPropertyChangedEvent& Proper
 	{
 		SetDistortionSource(DistortionStateSource);
 	}
+	else if (PropertyName == GET_MEMBER_NAME_CHECKED(ULensComponent, bOverrideCameraOverscan))
+	{
+		SetOverrideCameraOverscan(bOverrideCameraOverscan);
+	}
 
 	Super::PostEditChangeProperty(PropertyChangedEvent);
 }
@@ -531,6 +529,19 @@ void ULensComponent::SetOriginalTrackedComponentTransform(const FTransform& NewT
 	OriginalTrackedComponentRotation.Z = NewRotator.Yaw;
 }
 
+void ULensComponent::RestoreCameraOverscanDefaults()
+{
+	const UCineCameraComponent* CDO = UCineCameraComponent::StaticClass()->GetDefaultObject<UCineCameraComponent>();
+	UCineCameraComponent* CineCameraComponent = Cast<UCineCameraComponent>(TargetCameraComponent.GetComponent(GetOwner()));
+
+	if (IsValid(CineCameraComponent))
+	{
+		CineCameraComponent->Modify();
+		CineCameraComponent->Overscan = CDO->Overscan;
+		CineCameraComponent->bCropOverscan = CDO->bCropOverscan;
+	}
+}
+
 void ULensComponent::EvaluateFocalLength(UCineCameraComponent* CineCameraComponent)
 {
 	FFocalLengthInfo FocalLengthInfo;
@@ -770,6 +781,31 @@ bool ULensComponent::ShouldApplyDistortion() const
 void ULensComponent::SetApplyDistortion(bool bApply)
 {
 	bApplyDistortion = bApply;
+
+	if (bApplyDistortion == false)
+	{
+		if (UCineCameraComponent* const CineCameraComponent = Cast<UCineCameraComponent>(TargetCameraComponent.GetComponent(GetOwner())))
+		{
+			CleanupDistortion(CineCameraComponent);
+		}
+
+		RestoreCameraOverscanDefaults();
+	}
+}
+
+bool ULensComponent::GetOverrideCameraOverscan() const
+{
+	return bOverrideCameraOverscan;
+}
+
+void ULensComponent::SetOverrideCameraOverscan(bool bOverride)
+{
+	bOverrideCameraOverscan = bOverride;
+
+	if (bOverrideCameraOverscan == false)
+	{
+		RestoreCameraOverscanDefaults();
+	}
 }
 
 ULensDistortionModelHandlerBase* ULensComponent::GetLensDistortionHandler() const
@@ -795,17 +831,17 @@ EDistortionRenderingMode ULensComponent::GetDistortionRenderingMode() const
 		// If the handler unexpectedly returns Preferred, use default behavior based on IsForwardDistorting()
 		if (PreferredMode == EDistortionRenderingMode::Preferred)
 		{
-			// Reverse distorting models prefer PPM for accuracy, forward distorting prefer SVE
+			// Reverse distorting models prefer the plugin pass for accuracy, forward distorting models prefer TSR
 			return LensDistortionHandler->IsForwardDistorting() 
-				? EDistortionRenderingMode::SceneViewExtension 
-				: EDistortionRenderingMode::PostProcessMaterial;
+				? EDistortionRenderingMode::TemporalSuperResolution
+				: EDistortionRenderingMode::PassAfterTonemap;
 		}
 		
 		return PreferredMode;
 	}
 
-	// Default to SVE if no handler is found
-	return EDistortionRenderingMode::SceneViewExtension;
+	// Default to plugin pass if no handler is found
+	return EDistortionRenderingMode::PassAfterTonemap;
 }
 
 PRAGMA_DISABLE_DEPRECATION_WARNINGS
diff --git a/Engine/Plugins/VirtualProduction/LensComponent/Source/LensComponent/Public/LensComponent.h b/Engine/Plugins/VirtualProduction/LensComponent/Source/LensComponent/Public/LensComponent.h
--- a/Engine/Plugins/VirtualProduction/LensComponent/Source/LensComponent/Public/LensComponent.h
+++ b/Engine/Plugins/VirtualProduction/LensComponent/Source/LensComponent/Public/LensComponent.h
@@ -168,6 +168,14 @@ public:
 	UFUNCTION(BlueprintCallable, Category = "Lens Component")
 	UE_API void SetApplyDistortion(bool bApply);
 
+	/** Get whether the camera overscan is overriden. */
+	UFUNCTION(BlueprintPure, Category = "Lens Component")
+	UE_API bool GetOverrideCameraOverscan() const;
+
+	/** Set whether the camera overscan should be overriden. */
+	UFUNCTION(BlueprintCallable, Category = "Lens Component")
+	UE_API void SetOverrideCameraOverscan(bool bOverride);
+
 	/** Get the current lens model */
 	UFUNCTION(BlueprintPure, Category = "Lens Component")
 	UE_API TSubclassOf<ULensModel> GetLensModel() const;
@@ -275,6 +283,9 @@ private:
 	/** Set the original transform of the tracked component (and also the location and rotation vectors) */
 	UE_API void SetOriginalTrackedComponentTransform(const FTransform& NewTransform);
 
+	/** Convenience function to restore the camera overscan defaults. */
+	void RestoreCameraOverscanDefaults();
+
 protected:
 	/** Lens File used to drive distortion with current camera settings */
 	UPROPERTY(EditAnywhere, Category="Lens File", meta=(ShowOnlyInnerProperties))
@@ -307,11 +318,10 @@ protected:
 	/**
 	 * If checked, the camera's overscan value will be driven by the lens component to automatically compensate for distortion. 
 	 * The camera's overscan crop property will also be driven based on the distortion rendering mode:
-	 *   Disabled for Post Process Material
-	 *   Enabled for Scene View Extension
-	 * Note: The camera's overscan properties will not be automatically reset when the "Apply Distortion" or "Override Camera Overscan" properties are disabled
+	 *   Disabled for Post Process Material or pass after tonemap
+	 *   Enabled for Temporal Super Resolution
 	 */
-	UPROPERTY(EditAnywhere, BlueprintReadWrite, Category = "Distortion")
+	UPROPERTY(EditAnywhere, BlueprintReadWrite, Getter = GetOverrideCameraOverscan, Setter = SetOverrideCameraOverscan, Category = "Distortion")
 	bool bOverrideCameraOverscan = true;
 
 	/** The current lens model used for distortion */
@@ -323,7 +333,7 @@ protected:
 	FLensDistortionState DistortionState;
 
 	/** Whether to scale the computed overscan by the overscan percentage */
-	UPROPERTY(AdvancedDisplay, BlueprintReadWrite, Category = "Distortion", meta = (InlineEditConditionToggle))
+	UPROPERTY(EditAnywhere, AdvancedDisplay, BlueprintReadWrite, Category = "Distortion", meta = (InlineEditConditionToggle))
 	bool bScaleOverscan = false;
 
 	/** The percentage of the computed overscan that should be applied to the target camera */
